Moved EpicStudioAnalytics plugin to StudioTelemetry which is now public facing.
The intention is to provide a set of "in the box" telemetry hooks to our most common developer workflows for use internal Epic use and licensees
Moved startup of StudioTelemetry plugin much earlier to PostConfigInit stage so that sessions can start much earlier in the workflow
Added WIP Client support to telemetry plugin
Added data driven provider support for FAnlayiticProvidersET interfaces via BaseEngine.ini for common games and sample projects. Settings and URLs for EPic are only avaliable in Resttricted/NotForLicensee config folders for Lyra and Shooter game.
Added support for Horde telemetry and fixed up various problems with duplicate attributes being sent.
Added Core.VirtualAssets telemetry event
Added Core.Zen telemetry event
Added Core.IAS telemetry event
Added check for IOStoreOnDemand IsEnabled to avoid sending empty IAS events
Added FAnalyticsMulticastProvider to forward telemetry events to multiple providers contained within
Removed deprecated Fire_LoadingEvent from StudioAnalytics

diff --git a/Engine/Plugins/StudioTelemetry/Source/Private/AnalyticsProviderMulticast.cpp b/Engine/Plugins/StudioTelemetry/Source/Private/AnalyticsProviderMulticast.cpp
new file mode 100644
--- /dev/null
+++ b/Engine/Plugins/StudioTelemetry/Source/Private/AnalyticsProviderMulticast.cpp
@@ -0,0 +1,312 @@
+// Copyright Epic Games, Inc. All Rights Reserved.
+
+#include "AnalyticsProviderMulticast.h"
+#include "StudioTelemetryLog.h"
+#include "Analytics.h"
+#include "AnalyticsProviderMulticast.h"
+#include "AnalyticsProviderConfigurationDelegate.h"
+#include "Misc/ConfigCacheIni.h"
+#include "HttpModule.h"
+
+const FString TelemetrySection(TEXT("StudioTelemetry"));
+static FString ProviderSection;
+
+FString GetAnalyticsProviderConfiguration(const FString& Name, bool)
+{
+	FString Result;
+	GConfig->GetString(*ProviderSection, *Name, Result, GEngineIni);
+	return Result;
+}
+
+TSharedPtr<FAnalyticsProviderMulticast> FAnalyticsProviderMulticast::CreateAnalyticsProvider()
+{
+	return MakeShared<FAnalyticsProviderMulticast>();
+}
+
+TWeakPtr<IAnalyticsProviderET> FAnalyticsProviderMulticast::GetAnalyticsProvider(const FString& Name)
+{
+	TSharedPtr<IAnalyticsProviderET>* ProviderPtr = Providers.Find(Name);
+	return ProviderPtr != nullptr ? *ProviderPtr : TSharedPtr<IAnalyticsProviderET>();
+}
+
+FAnalyticsProviderMulticast::FAnalyticsProviderMulticast()
+{
+	TArray<FString> SectionNames;
+	
+	if (GConfig->GetSectionNames(GEngineIni, SectionNames))
+	{
+		for (const FString& SectionName : SectionNames)
+		{
+			if (SectionName.Find(TelemetrySection) != INDEX_NONE)
+			{
+				ProviderSection = SectionName;
+
+				FString UsageType;
+
+				// Validate the usage type is for this build type
+				if (GConfig->GetString(*ProviderSection, TEXT("UsageType"), UsageType, GEngineIni))
+				{
+#if WITH_EDITOR
+					// Must specify a Editor usage type for this type build
+					if (UsageType.Find(TEXT("Editor")) == INDEX_NONE)
+					{
+						continue;
+					}
+#elif WITH_SERVER_CODE
+					// Must specify a Server usage type for this type build
+					if (UsageType.Find(TEXT("Server")) == INDEX_NONE)
+					{
+						continue;
+					}
+#else
+					// Must specify a Client or Program usage type for this type build
+					if (UsageType.Find(TEXT("Program")) == INDEX_NONE && UsageType.Find(TEXT("Client")) == INDEX_NONE)
+					{
+						continue;
+					}
+#endif
+				}
+				else
+				{
+					// Must always specify a usage type
+					UE_LOG(LogStudioTelemetry, Error, TEXT("There must be a valid UsageType specified for analytics provider %s"), *ProviderSection);
+					continue;
+				}
+
+				FString ProviderType;
+
+				if (GConfig->GetString(*ProviderSection, TEXT("ProviderType"), ProviderType, GEngineIni))
+				{
+					TSharedPtr<IAnalyticsProviderET> Provider;
+
+					FString Name = GetAnalyticsProviderConfiguration("Name", true);
+
+					if ( Name.IsEmpty() )
+					{ 
+						UE_LOG(LogStudioTelemetry, Error, TEXT("There must be a valid Name specified for analytics provider %s."), *ProviderSection);
+						continue;
+					}
+					else if (Providers.Find(Name) )
+					{
+						UE_LOG(LogStudioTelemetry, Warning, TEXT("An analytics provider with name %s already exists."), *Name);
+						continue;
+					}
+					
+					if (ProviderType.Equals(TEXT("FAnalyticsProviderET"), ESearchCase::IgnoreCase))
+					{
+						Provider = FAnalyticsET::Get().CreateAnalyticsProviderET(FAnalyticsProviderConfigurationDelegate::CreateStatic(&GetAnalyticsProviderConfiguration));
+					}
+					
+					if (Provider.IsValid())
+					{
+						UE_LOG(LogStudioTelemetry, Display, TEXT("Created a %s analytics provider %s from configuration %s [%s]"), *ProviderType, *Name, *GEngineIni, *ProviderSection);
+
+						Providers.Add(Name, Provider);
+					}
+				}
+				else
+				{
+					UE_LOG(LogStudioTelemetry, Error, TEXT("There must be a valid ProviderType specified for analytics provider %s"), *ProviderSection);
+				}
+			}
+		}
+	}
+}
+
+void FAnalyticsProviderMulticast::SetAppID(FString&& InAppID)
+{
+	Config.APIKeyET = InAppID;
+
+	for (TProviders::TConstIterator it(Providers); it; ++it)
+	{
+		(*it).Value->SetAppID(CopyTemp(InAppID));
+	}
+}
+
+void FAnalyticsProviderMulticast::SetAppVersion(FString&& InAppVersion)
+{
+	Config.AppVersionET = InAppVersion;
+
+	for (TProviders::TConstIterator it(Providers); it; ++it)
+	{
+		(*it).Value->SetAppVersion(CopyTemp(InAppVersion));
+	}
+}
+
+bool FAnalyticsProviderMulticast::SetSessionID(const FString& InSessionID)
+{
+	SessionID = InSessionID;
+
+	bool bResult = true;
+
+	for (TProviders::TConstIterator it(Providers); it; ++it)
+	{
+		bResult &= (*it).Value->SetSessionID(InSessionID);
+	}
+
+	return bResult;
+}
+
+FString FAnalyticsProviderMulticast::GetSessionID() const
+{
+	return SessionID;
+}
+
+void FAnalyticsProviderMulticast::SetUserID(const FString& InUserID)
+{
+	UserID = InUserID;
+
+	for (TProviders::TConstIterator it(Providers); it; ++it)
+	{
+		(*it).Value->SetUserID(InUserID);
+	}
+}
+
+FString FAnalyticsProviderMulticast::GetUserID() const
+{
+	return UserID;
+}
+
+const FAnalyticsET::Config& FAnalyticsProviderMulticast::GetConfig() const
+{
+	return Config;
+}
+
+void FAnalyticsProviderMulticast::SetURLEndpoint(const FString& UrlEndpoint, const TArray<FString>& AltDomains)
+{
+	for (TProviders::TConstIterator it(Providers); it; ++it)
+	{
+		(*it).Value->SetURLEndpoint(UrlEndpoint, AltDomains);
+	}
+}
+
+void FAnalyticsProviderMulticast::SetHeader(const FString& HeaderName, const FString& HeaderValue)
+{
+	for (TProviders::TConstIterator it(Providers); it; ++it)
+	{
+		(*it).Value->SetHeader(HeaderName, HeaderValue);
+	}
+}
+
+void FAnalyticsProviderMulticast::SetEventCallback(const OnEventRecorded& InCallback)
+{
+	OnEventRecordedCallback = InCallback;
+}
+
+void FAnalyticsProviderMulticast::FlushEvents()
+{
+	for (TProviders::TConstIterator it(Providers); it; ++it)
+	{
+		(*it).Value->FlushEvents();
+	}
+}
+
+void FAnalyticsProviderMulticast::BlockUntilFlushed(float InTimeoutSec)
+{
+	for (TProviders::TConstIterator it(Providers); it; ++it)
+	{
+		(*it).Value->BlockUntilFlushed(InTimeoutSec);
+	}
+}
+
+bool FAnalyticsProviderMulticast::StartSession(FString InSessionID, const TArray<FAnalyticsEventAttribute>& Attributes)
+{
+	SetSessionID(InSessionID);
+
+	bool bResult = true;
+
+	for (TProviders::TConstIterator it(Providers); it; ++it)
+	{
+		bResult &= (*it).Value->StartSession(InSessionID, Attributes);
+	}
+
+	return bResult;
+}
+
+void FAnalyticsProviderMulticast::SetShouldRecordEventFunc(const ShouldRecordEventFunction& InShouldRecordEventFunc)
+{
+	ShouldRecordEventFunc = InShouldRecordEventFunc;
+}
+
+bool FAnalyticsProviderMulticast::ShouldRecordEvent(const FString& EventName) const
+{
+	return ShouldRecordEventFunc ? ShouldRecordEventFunc(*this, EventName) : true;
+}
+
+void FAnalyticsProviderMulticast::SetDefaultEventAttributes(TArray<FAnalyticsEventAttribute>&& Attributes)
+{
+	DefaultEventAttributes = Attributes;
+
+	for (TProviders::TConstIterator it(Providers); it; ++it)
+	{
+		(*it).Value->SetDefaultEventAttributes(CopyTemp(DefaultEventAttributes));
+	}
+}
+
+TArray<FAnalyticsEventAttribute> FAnalyticsProviderMulticast::GetDefaultEventAttributesSafe() const
+{
+	return DefaultEventAttributes;
+}
+
+int32 FAnalyticsProviderMulticast::GetDefaultEventAttributeCount() const
+{
+	return DefaultEventAttributes.Num();
+}
+
+FAnalyticsEventAttribute FAnalyticsProviderMulticast::GetDefaultEventAttribute(int AttributeIndex) const
+{
+	return DefaultEventAttributes[AttributeIndex];
+}
+
+bool FAnalyticsProviderMulticast::StartSession(const TArray<FAnalyticsEventAttribute>& Attributes)
+{
+	bool bResult = true;
+
+	for (TProviders::TConstIterator it(Providers);it;++it)
+	{
+		bResult &= (*it).Value->StartSession(Attributes);
+	}
+	return bResult;
+}
+
+void FAnalyticsProviderMulticast::EndSession()
+{
+	for (TProviders::TConstIterator it(Providers);it;++it)
+	{
+		TSharedPtr<IAnalyticsProviderET> Provider = (*it).Value;
+
+		Provider->EndSession();
+		Provider.Reset();
+	}
+
+	Providers.Reset();
+}
+
+void FAnalyticsProviderMulticast::RecordEvent(FString&& EventName, const TArray<FAnalyticsEventAttribute>& Attributes)
+{
+	if (ShouldRecordEvent(EventName))
+	{
+#if UE_BUILD_DEBUG || UE_BUILD_DEVELOPMENT
+		// Expose events that have duplicate aatibute names. This is is not handled by the analytics backends in any reliable way.
+		for (int32 index0 = 0; index0 < Attributes.Num(); ++index0)
+		{
+			for (int32 index1 = index0 + 1; index1 < Attributes.Num(); ++index1)
+			{
+				checkf(Attributes[index0].GetName() != Attributes[index1].GetName(), TEXT("Duplicate Attributes Found For Event %s %s==%s"), *EventName, *Attributes[index0].GetName(), *Attributes[index1].GetName());
+			}
+		}
+#endif
+
+		for (TProviders::TConstIterator it(Providers); it; ++it)
+		{
+			(*it).Value->RecordEvent(EventName, Attributes);
+		}
+
+		// Notify any callbacks
+		if (OnEventRecordedCallback)
+		{
+			OnEventRecordedCallback(EventName, Attributes, true);
+		}
+	}
+
+}
diff --git a/Engine/Plugins/StudioTelemetry/Source/Private/AnalyticsProviderMulticast.h b/Engine/Plugins/StudioTelemetry/Source/Private/AnalyticsProviderMulticast.h
new file mode 100644
--- /dev/null
+++ b/Engine/Plugins/StudioTelemetry/Source/Private/AnalyticsProviderMulticast.h
@@ -0,0 +1,61 @@
+// Copyright Epic Games, Inc. All Rights Reserved.
+
+#pragma once
+
+#include "AnalyticsProviderConfigurationDelegate.h"
+#include "IAnalyticsProviderET.h"
+
+/**
+ * Implementation of the IAnalyticsProviderET interface that forwards the API calls to an array of IAnalyticsProviderET interfaces
+ */
+class FAnalyticsProviderMulticast : public IAnalyticsProviderET
+{
+public:
+
+	using TProviders = TMap<FString, TSharedPtr<IAnalyticsProviderET>>;
+
+	FAnalyticsProviderMulticast();
+
+	static TSharedPtr<FAnalyticsProviderMulticast> CreateAnalyticsProvider();
+
+	virtual void SetShouldRecordEventFunc(const ShouldRecordEventFunction& ShouldRecordEventFunc) override;
+	virtual bool StartSession(FString InSessionID, const TArray<FAnalyticsEventAttribute>& Attributes = {}) override;
+	virtual bool StartSession(const TArray<FAnalyticsEventAttribute>& Attributes = {}) override;
+	virtual void EndSession() override;
+
+	virtual bool ShouldRecordEvent(const FString& EventName) const override;
+	virtual void FlushEvents() override;
+	virtual void BlockUntilFlushed(float InTimeoutSec) override;
+
+	virtual void SetDefaultEventAttributes(TArray<FAnalyticsEventAttribute>&& Attributes)  override;
+	virtual TArray<FAnalyticsEventAttribute> GetDefaultEventAttributesSafe() const override;
+	virtual int32 GetDefaultEventAttributeCount() const  override;
+	virtual FAnalyticsEventAttribute GetDefaultEventAttribute(int AttributeIndex) const  override;
+
+	virtual void SetURLEndpoint(const FString& UrlEndpoint, const TArray<FString>& AltDomains)  override;
+	virtual void SetHeader(const FString& HeaderName, const FString& HeaderValue) override;
+	virtual void SetAppID(FString&& AppID) override;
+	virtual void SetAppVersion(FString&& AppVersion) override;
+	virtual void SetEventCallback(const OnEventRecorded& Callback)  override;
+	virtual bool SetSessionID(const FString& InSessionID) override;
+	virtual void SetUserID(const FString& InUserID) override;
+
+	virtual FString GetSessionID() const override;
+	virtual FString GetUserID() const override;
+	virtual const FAnalyticsET::Config& GetConfig() const  override;
+
+	virtual void RecordEvent(FString&& EventName, const TArray<FAnalyticsEventAttribute>& Attributes) override;
+	
+	bool HasValidProviders() const { return Providers.Num() > 0; }
+	TWeakPtr<IAnalyticsProviderET> GetAnalyticsProvider(const FString& Name);
+
+private:
+
+	TProviders								Providers;
+	FAnalyticsET::Config					Config;
+	FString									UserID;
+	FString									SessionID;
+	TArray<FAnalyticsEventAttribute>		DefaultEventAttributes;
+	ShouldRecordEventFunction				ShouldRecordEventFunc;
+	OnEventRecorded							OnEventRecordedCallback;
+};
diff --git a/Engine/Plugins/StudioTelemetry/Source/Private/StudioTelemetry.cpp b/Engine/Plugins/StudioTelemetry/Source/Private/StudioTelemetry.cpp
new file mode 100644
--- /dev/null
+++ b/Engine/Plugins/StudioTelemetry/Source/Private/StudioTelemetry.cpp
@@ -0,0 +1,201 @@
+// Copyright Epic Games, Inc. All Rights Reserved.
+
+#include "StudioTelemetry.h"
+#include "StudioTelemetryLog.h"
+
+#if WITH_EDITOR
+#include "StudioTelemetryEditor.h"
+#endif
+
+#if !UE_BUILD_SHIPPING
+#include "Horde.h"
+#endif
+
+#include "Analytics.h"
+#include "AnalyticsProviderMulticast.h"
+#include "BuildSettings.h"
+#include "RHI.h"
+
+#include "HAL/PlatformTime.h"
+#include "HAL/PlatformProcess.h"
+#include "HAL/Thread.h"
+
+#include "Misc/CommandLine.h"
+#include "Misc/App.h"
+#include "Misc/EngineVersion.h"
+#include "Misc/EngineBuildSettings.h"
+
+#include "GenericPlatform/GenericPlatformMisc.h"
+
+DEFINE_LOG_CATEGORY(LogStudioTelemetry);
+
+IMPLEMENT_MODULE(FStudioTelemetry, StudioTelemetry)
+
+FStudioTelemetry& FStudioTelemetry::Get()
+{
+	static FStudioTelemetry StudioTelemetryInstance;
+	return StudioTelemetryInstance;
+}
+
+void FStudioTelemetry::SetOnEventRecordedCallback(OnEventRecorded InOnEventRecordedCallback )
+{
+	OnEventRecordedCallback = InOnEventRecordedCallback;
+
+	// If the provider already exists then set the callback
+	if (AnalyticsProvider.IsValid())
+	{	
+		AnalyticsProvider->SetEventCallback(InOnEventRecordedCallback);	
+	}
+}
+
+void FStudioTelemetry::StartupModule()
+{
+	UE_LOG(LogStudioTelemetry, Display, TEXT("Starting StudioTelemetry Module"));
+
+	// Create the provider and start the analytics session
+	FStudioTelemetry::Get().StartSession();
+
+#if WITH_EDITOR
+	// Initialize the analytics subsystems 
+	FStudioTelemetryEditor::Get().Initialize();
+#endif
+}
+
+void FStudioTelemetry::ShutdownModule()
+{
+#if WITH_EDITOR
+	// Shutdown the analytics subsystems
+	FStudioTelemetryEditor::Get().Shutdown();
+#endif
+
+	// End the session and destroy analytics provider
+	FStudioTelemetry::Get().EndSession();
+
+	UE_LOG(LogStudioTelemetry, Display, TEXT("Shutdown StudioTelemetry Module"));
+}
+
+void FStudioTelemetry::EndSession()
+{
+	if (AnalyticsFlowTracker.IsValid())
+	{
+		AnalyticsFlowTracker->EndSession();
+		AnalyticsFlowTracker.Reset();
+	}
+
+	if (AnalyticsProvider.IsValid())
+	{
+		AnalyticsProvider->EndSession();
+		AnalyticsProvider.Reset();
+	}
+
+	UE_LOG(LogStudioTelemetry, Log, TEXT("Ended StudioTelemetry Session"));
+}
+
+void FStudioTelemetry::StartSession()
+{
+	AnalyticsProvider = FAnalyticsProviderMulticast::CreateAnalyticsProvider();
+
+	if (AnalyticsProvider.IsValid())
+	{
+		TArray<FAnalyticsEventAttribute> DefaultEventAttributes;
+
+		const FString UserID = FPlatformProcess::UserName(false);
+		const FString ProjectName = FApp::GetProjectName();
+
+		FString ProjectIDString;
+		GConfig->GetString(TEXT("/Script/EngineSettings.GeneralProjectSettings"), TEXT("ProjectID"), ProjectIDString, GGameIni);
+
+		FGuid ProjectID(ProjectIDString);
+
+		FString SessionLabel;
+		FParse::Value(FCommandLine::Get(), TEXT("SessionLabel="), SessionLabel);
+
+		// Set the default event attributes
+		DefaultEventAttributes.Emplace(TEXT("ProjectName"), ProjectName);
+		DefaultEventAttributes.Emplace(TEXT("ProjectID"), ProjectID);
+		DefaultEventAttributes.Emplace(TEXT("User_ID"), UserID);
+		DefaultEventAttributes.Emplace(TEXT("Application_Commandline"), FCommandLine::Get());
+
+		DefaultEventAttributes.Emplace(TEXT("Session_Label"), SessionLabel);
+		DefaultEventAttributes.Emplace(TEXT("Session_StartUTC"), FDateTime::UtcNow().ToUnixTimestampDecimal());
+
+		DefaultEventAttributes.Emplace(TEXT("Build_Configuration"), LexToString(FApp::GetBuildConfiguration()));
+		DefaultEventAttributes.Emplace(TEXT("Build_IsInternalBuild"), FEngineBuildSettings::IsInternalBuild());
+		DefaultEventAttributes.Emplace(TEXT("Build_IsPerforceBuild"), FEngineBuildSettings::IsPerforceBuild());
+		DefaultEventAttributes.Emplace(TEXT("Build_IsPromotedBuild"), FApp::GetEngineIsPromotedBuild() == 0 ? false : true);
+		DefaultEventAttributes.Emplace(TEXT("Build_BranchName"), FApp::GetBranchName());
+		DefaultEventAttributes.Emplace(TEXT("Build_Changelist"), BuildSettings::GetCurrentChangelist());
+
+		DefaultEventAttributes.Emplace(TEXT("Hardware_GPU"), GRHIAdapterName);
+		DefaultEventAttributes.Emplace(TEXT("Hardware_CPU"), FPlatformMisc::GetCPUBrand());
+		DefaultEventAttributes.Emplace(TEXT("Hardware_CPU_Cores_Physical"), FPlatformMisc::NumberOfCores());
+		DefaultEventAttributes.Emplace(TEXT("Hardware_CPU_Cores_Logical"), FPlatformMisc::NumberOfCoresIncludingHyperthreads());
+		DefaultEventAttributes.Emplace(TEXT("Hardware_RAM"), static_cast<uint64>(FPlatformMemory::GetStats().TotalPhysical));
+
+		DefaultEventAttributes.Emplace(TEXT("Config_IsEditor"), GIsEditor);
+		DefaultEventAttributes.Emplace(TEXT("Config_IsUnattended"), FApp::IsUnattended());
+		DefaultEventAttributes.Emplace(TEXT("Config_IsBuildMachine"), GIsBuildMachine);
+		DefaultEventAttributes.Emplace(TEXT("Config_IsRunningCommandlet"), IsRunningCommandlet());
+
+#if !UE_BUILD_SHIPPING
+		DefaultEventAttributes.Emplace(TEXT("Horde_TemplateID"), FHorde::GetTemplateId());
+		DefaultEventAttributes.Emplace(TEXT("Horde_TemplateName"), FHorde::GetTemplateName());
+		DefaultEventAttributes.Emplace(TEXT("Horde_JobURL"), FHorde::GetJobURL());
+		DefaultEventAttributes.Emplace(TEXT("Horde_JobID"), FHorde::GetJobId());
+		DefaultEventAttributes.Emplace(TEXT("Horde_StepName"), FHorde::GetStepName());
+		DefaultEventAttributes.Emplace(TEXT("Horde_StepID"), FHorde::GetStepId());
+		DefaultEventAttributes.Emplace(TEXT("Horde_StepURL"), FHorde::GetStepURL());
+		DefaultEventAttributes.Emplace(TEXT("Horde_BatchID"), FHorde::GetBatchId());
+#endif
+
+		// Set up the analytics provider
+		AnalyticsProvider->SetUserID(UserID);
+		AnalyticsProvider->SetDefaultEventAttributes(MoveTemp(DefaultEventAttributes));
+		AnalyticsProvider->SetEventCallback(OnEventRecordedCallback);
+		
+		// Start the analytics session
+		AnalyticsProvider->StartSession();
+
+		// Make the flow tracker and start the session
+		AnalyticsFlowTracker = MakeShared<FAnalyticsFlowTracker>();
+		AnalyticsFlowTracker->SetProvider(AnalyticsProvider);
+		AnalyticsFlowTracker->StartSession();
+
+		UE_LOG(LogStudioTelemetry, Log, TEXT("Started StudioTelemetry Session"));
+	}
+}
+
+void FStudioTelemetry::RecordEvent(const FString& EventName, const TArray<FAnalyticsEventAttribute>& Attributes)
+{
+	if (AnalyticsProvider.IsValid())
+	{
+		FScopeLock ScopeLock(&CriticalSection);
+		AnalyticsProvider->RecordEvent(CopyTemp(EventName), Attributes);
+	}
+}
+
+void FStudioTelemetry::RecordEvent(const FString& ProviderName, const FString& EventName, const TArray<FAnalyticsEventAttribute>& Attributes)
+{
+	FScopeLock ScopeLock(&CriticalSection);
+	TSharedPtr<IAnalyticsProvider> NamedProvider = GetProvider(ProviderName).Pin();
+
+	if (NamedProvider.IsValid())
+	{
+		NamedProvider->RecordEvent(EventName, Attributes);
+	}
+}
+
+TWeakPtr<IAnalyticsProvider> FStudioTelemetry::GetProvider()
+{
+	return AnalyticsProvider;
+}
+
+TWeakPtr<IAnalyticsProvider> FStudioTelemetry::GetProvider(const FString& Name)
+{
+	return AnalyticsProvider.IsValid()? AnalyticsProvider->GetAnalyticsProvider(Name) : TWeakPtr<IAnalyticsProvider>();
+}
+
+TWeakPtr<FAnalyticsFlowTracker> FStudioTelemetry::GetFlowTracker()
+{
+	return AnalyticsFlowTracker;
+}
\ No newline at end of file
diff --git a/Engine/Plugins/StudioTelemetry/Source/Private/StudioTelemetryEditor.cpp b/Engine/Plugins/StudioTelemetry/Source/Private/StudioTelemetryEditor.cpp
new file mode 100644
--- /dev/null
+++ b/Engine/Plugins/StudioTelemetry/Source/Private/StudioTelemetryEditor.cpp
@@ -0,0 +1,865 @@
+// Copyright Epic Games, Inc. All Rights Reserved.
+
+#include "StudioTelemetryEditor.h"
+
+#if WITH_EDITOR
+
+#include "StudioTelemetry.h"
+#include "AnalyticsFlowTracker.h"
+#include "CollectionManagerModule.h"
+#include "ContentBrowserModule.h"
+#include "TelemetryRouter.h"
+#include "Policies/CondensedJsonPrintPolicy.h"
+#include "AssetRegistry/AssetRegistryTelemetry.h"
+#include "Editor.h"
+#include "Editor/UnrealEdEngine.h"
+#include "UnrealEdGlobals.h"
+#include "CookOnTheSide/CookOnTheFlyServer.h"
+#include "Subsystems/AssetEditorSubsystem.h"
+#include "Misc/FeedbackContext.h"
+#include "ProfilingDebugging/CookStats.h"
+#include "DerivedDataCacheInterface.h"
+#include "DerivedDataCacheUsageStats.h"
+#include "Virtualization/VirtualizationSystem.h"
+#include "FileHelpers.h"
+#include "Experimental/ZenServerInterface.h"
+#include "ContentBrowserTelemetry.h"
+#include "ShaderStats.h"
+#include "UObject/ICookInfo.h"
+#include "IO/IoStoreOnDemand.h"
+
+UE_DISABLE_OPTIMIZATION_SHIP
+
+namespace Private
+{
+	const FName ContentBrowserModuleName = TEXT("ContentBrowser");
+	
+	// Json writer subclass to allow us to avoid using a SharedPtr to write basic Json.
+	typedef TCondensedJsonPrintPolicy<TCHAR> FPrintPolicy;
+	class FAnalyticsJsonWriter : public TJsonStringWriter<FPrintPolicy>
+	{
+	public:
+		explicit FAnalyticsJsonWriter(FString* Out) : TJsonStringWriter<FPrintPolicy>(Out, 0)
+		{
+		}
+	};
+}
+
+const TCHAR* LexToString(ECollectionTelemetryAssetAddedWorkflow Enum)
+{
+	switch(Enum)
+	{
+	case ECollectionTelemetryAssetAddedWorkflow::ContextMenu: return TEXT("ContextMenu");
+	case ECollectionTelemetryAssetAddedWorkflow::DragAndDrop: return TEXT("DragAndDrop");
+	default: return TEXT("");
+	}
+}
+	
+const TCHAR* LexToString(ECollectionTelemetryAssetRemovedWorkflow Enum)
+{
+	switch(Enum)
+	{
+	case ECollectionTelemetryAssetRemovedWorkflow::ContextMenu: return TEXT("ContextMenu");
+	default: return TEXT("");
+	}
+}
+
+template<typename T>
+FString AnalyticsOptionalToStringOrNull(const TOptional<T>& Opt)
+{
+	return Opt.IsSet() ? AnalyticsConversionToString(Opt.GetValue()) : FString(TEXT("null"));
+}
+
+FStudioTelemetryEditor& FStudioTelemetryEditor::Get()
+{
+	static FStudioTelemetryEditor StudioTelemetryEditorInstance = FStudioTelemetryEditor();
+	return StudioTelemetryEditorInstance;
+}
+
+void FStudioTelemetryEditor::RecordEvent_Cooking(TArray<FAnalyticsEventAttribute> Attributes)
+{
+#if ENABLE_COOK_STATS
+
+	const int SchemaVersion = 2;
+
+	Attributes.Emplace(TEXT("SchemaVersion"), SchemaVersion);
+
+	TMap<FString,FAnalyticsEventAttribute> CookAttributes;
+	
+	// Sends each cook stat to the studio analytics system.
+	auto GatherAnalyticsAttributes = [&CookAttributes](const FString& StatName, const TArray<FCookStatsManager::StringKeyValue>& StatAttributes)
+	{
+		for (const auto& Attr : StatAttributes)
+		{
+			const FString FormattedAttrName = (StatName + "_" + Attr.Key).Replace(TEXT("."), TEXT("_"));
+
+			if (CookAttributes.Find(FormattedAttrName)==nullptr)
+			{
+				CookAttributes.Emplace(FormattedAttrName, Attr.Value.IsNumeric() ? FAnalyticsEventAttribute(FormattedAttrName, FCString::Atof(*Attr.Value)) : FAnalyticsEventAttribute(FormattedAttrName, Attr.Value));
+			}
+		}
+	};
+
+	// Now actually grab the stats 
+	FCookStatsManager::LogCookStats(GatherAnalyticsAttributes);
+
+	// Add the values to the attributes
+	for (TMap<FString, FAnalyticsEventAttribute>::TConstIterator it(CookAttributes); it; ++it)
+	{
+		Attributes.Emplace((*it).Value);
+	}
+
+	// Gather the DDC summary stats
+	FDerivedDataCacheSummaryStats SummaryStats;
+
+	GatherDerivedDataCacheSummaryStats(SummaryStats);
+
+	// Append to the attributes
+	for (const FDerivedDataCacheSummaryStat& Stat : SummaryStats.Stats)
+	{
+		FString AttributeName = TEXT("DDC_Summary") + Stat.Key.Replace(TEXT("."), TEXT("_"));
+
+		if (Stat.Value.IsNumeric())
+		{
+			Attributes.Emplace(AttributeName, FCString::Atof(*Stat.Value));
+		}
+		else
+		{
+			Attributes.Emplace(AttributeName, Stat.Value);
+		}
+	}
+
+#if UE_WITH_ZEN
+	// Gather Zen analytics
+	if (UE::Zen::IsDefaultServicePresent())
+	{
+		UE::Zen::GetDefaultServiceInstance().GatherAnalytics(Attributes);
+	}
+#endif
+
+	if (UE::Virtualization::IVirtualizationSystem::Get().IsEnabled())
+	{
+		// Gather Virtualization analytics
+		UE::Virtualization::IVirtualizationSystem::Get().GatherAnalytics(Attributes);
+	}
+
+	FShaderStatsFunctions::GatherShaderAnalytics(Attributes);
+	
+	FStudioTelemetry::Get().RecordEvent(TEXT("Core.Cooking"), Attributes);
+#endif
+}
+
+void FStudioTelemetryEditor::RecordEvent_Loading(const FString& LoadingName, double LoadingSeconds, TArray<FAnalyticsEventAttribute> Attributes )
+{
+	const int SchemaVersion = 2;
+
+	Attributes.Emplace(TEXT("SchemaVersion"), SchemaVersion);
+	Attributes.Emplace(TEXT("LoadingName"), LoadingName);
+	Attributes.Emplace(TEXT("LoadingSeconds"), LoadingSeconds);
+	
+#if ENABLE_COOK_STATS
+
+#if UE_WITH_ZEN
+	// Gather Zen analytics
+	if (UE::Zen::IsDefaultServicePresent())
+	{
+		UE::Zen::GetDefaultServiceInstance().GatherAnalytics(Attributes);
+	}
+#endif
+
+	if (UE::Virtualization::IVirtualizationSystem::Get().IsEnabled())
+	{
+		// Gather Virtualization analytics
+		UE::Virtualization::IVirtualizationSystem::Get().GatherAnalytics(Attributes);
+	}
+
+	// Gather the DDC summary stats
+	FDerivedDataCacheSummaryStats SummaryStats;
+
+	GatherDerivedDataCacheSummaryStats(SummaryStats);
+
+	// Append to the attributes
+	for (const FDerivedDataCacheSummaryStat& Stat : SummaryStats.Stats)
+	{
+		FString AttributeName = TEXT("DDC_Summary_") + Stat.Key.Replace(TEXT("."), TEXT("_"));
+
+		if (Stat.Value.IsNumeric())
+		{
+			Attributes.Emplace(AttributeName, FCString::Atof(*Stat.Value));
+		}
+		else
+		{
+			Attributes.Emplace(AttributeName, Stat.Value);
+		}
+	}
+#endif			
+
+	FStudioTelemetry::Get().RecordEvent(TEXT("Core.Loading"), Attributes);
+}
+
+void FStudioTelemetryEditor::RecordEvent_DDCResource(const FString& Context, TArray<FAnalyticsEventAttribute> Attributes)
+{
+#if ENABLE_COOK_STATS
+	// Gather the latest resource stats
+	TArray<FDerivedDataCacheResourceStat> ResourceStats;
+
+	GatherDerivedDataCacheResourceStats(ResourceStats);
+
+	const int SchemaVersion = 2;
+
+	Attributes.Emplace(TEXT("SchemaVersion"), SchemaVersion);
+	Attributes.Emplace(TEXT("LoadingName"), Context);
+	
+	// Send a resource event per asset type
+	for (const FDerivedDataCacheResourceStat& Stat : ResourceStats)
+	{
+		const int64 TotalCount = Stat.BuildCount + Stat.LoadCount;
+		
+		TArray<FAnalyticsEventAttribute> EventAttributes = Attributes;
+
+		EventAttributes.Emplace(TEXT("Asset_Type"), Stat.AssetType);
+		EventAttributes.Emplace(TEXT("Load_Count"), Stat.LoadCount);
+		EventAttributes.Emplace(TEXT("Load_TimeSec"), Stat.LoadTimeSec);
+		EventAttributes.Emplace(TEXT("Load_SizeMB"), Stat.LoadSizeMB);
+		EventAttributes.Emplace(TEXT("Build_Count"), Stat.BuildCount);
+		EventAttributes.Emplace(TEXT("Build_TimeSec"), Stat.BuildTimeSec);
+		EventAttributes.Emplace(TEXT("Build_SizeMB"), Stat.BuildSizeMB);
+		EventAttributes.Emplace(TEXT("Efficiency"), TotalCount>0? double(Stat.LoadCount)/double(TotalCount): 0.0 );
+		EventAttributes.Emplace(TEXT("Thread_TimeSec"), Stat.GameThreadTimeSec);
+
+		FStudioTelemetry::Get().RecordEvent(TEXT("Core.DDC.Resource"), EventAttributes);
+	}
+#endif			
+}
+
+
+void FStudioTelemetryEditor::RecordEvent_DDCSummary(const FString& Context, TArray<FAnalyticsEventAttribute> Attributes)
+{
+#if ENABLE_COOK_STATS
+	const int SchemaVersion = 2;
+
+	Attributes.Emplace(TEXT("SchemaVersion"), SchemaVersion);
+	Attributes.Emplace(TEXT("LoadingName"), Context);
+
+	// Gather the summary stats
+	FDerivedDataCacheSummaryStats SummaryStats;
+
+	GatherDerivedDataCacheSummaryStats(SummaryStats);
+
+	// Append to the attributes
+	for (const FDerivedDataCacheSummaryStat& Stat : SummaryStats.Stats)
+	{
+		FString AttributeName = Stat.Key.Replace(TEXT("."), TEXT("_"));
+
+		if (Stat.Value.IsNumeric())
+		{
+			Attributes.Emplace(AttributeName, FCString::Atof(*Stat.Value));
+		}
+		else
+		{
+			Attributes.Emplace(AttributeName, Stat.Value);
+		}
+	}
+
+	FStudioTelemetry::Get().RecordEvent(TEXT("Core.DDC.Summary"), Attributes);
+#endif			
+}
+
+void FStudioTelemetryEditor::RecordEvent_IAS(const FString& Context, TArray<FAnalyticsEventAttribute> Attributes)
+{
+	// Gather the summary stats
+	FDerivedDataCacheSummaryStats SummaryStats;
+
+	using namespace UE::IO::IAS;
+	if (FIoStoreOnDemandModule* OnDemandModule = FModuleManager::Get().GetModulePtr<FIoStoreOnDemandModule>("IoStoreOnDemand"))
+	{
+		if (OnDemandModule->IsEnabled())
+		{
+			const int SchemaVersion = 1;
+
+			Attributes.Emplace(TEXT("SchemaVersion"), SchemaVersion);
+			Attributes.Emplace(TEXT("LoadingName"), Context);
+
+			OnDemandModule->ReportAnalytics(Attributes);
+			FStudioTelemetry::Get().RecordEvent(TEXT("Core.IAS"), Attributes);
+		}		
+	}
+}
+
+void FStudioTelemetryEditor::RecordEvent_Zen(const FString& Context, TArray<FAnalyticsEventAttribute> Attributes)
+{
+#if UE_WITH_ZEN
+	// Gather Zen analytics
+	if (UE::Zen::IsDefaultServicePresent())
+	{
+		const int SchemaVersion = 1;
+
+		Attributes.Emplace(TEXT("SchemaVersion"), SchemaVersion);
+		Attributes.Emplace(TEXT("LoadingName"), Context);
+
+		UE::Zen::GetDefaultServiceInstance().GatherAnalytics(Attributes);
+		FStudioTelemetry::Get().RecordEvent(TEXT("Core.Zen"), Attributes);
+	}
+#endif
+}
+
+void FStudioTelemetryEditor::RecordEvent_VirtualAssets(const FString& Context, TArray<FAnalyticsEventAttribute> Attributes)
+{
+	if (UE::Virtualization::IVirtualizationSystem::Get().IsEnabled())
+	{
+		const int SchemaVersion = 1;
+
+		Attributes.Emplace(TEXT("SchemaVersion"), SchemaVersion);
+		Attributes.Emplace(TEXT("LoadingName"), Context);
+
+		// Gather Virtualization analytics
+		UE::Virtualization::IVirtualizationSystem::Get().GatherAnalytics(Attributes);
+
+		FStudioTelemetry::Get().RecordEvent(TEXT("Core.VirtualAssets"), Attributes);
+	}
+}
+
+void FStudioTelemetryEditor::RecordEvent_CoreSystems(const FString& Context, TArray<FAnalyticsEventAttribute> Attributes)
+{
+	FStudioTelemetryEditor::RecordEvent_DDCResource(Context, Attributes);
+	FStudioTelemetryEditor::RecordEvent_DDCSummary(Context, Attributes);
+	FStudioTelemetryEditor::RecordEvent_IAS(Context, Attributes);
+	FStudioTelemetryEditor::RecordEvent_Zen(Context, Attributes);
+	FStudioTelemetryEditor::RecordEvent_VirtualAssets(Context, Attributes);
+}
+
+void FStudioTelemetryEditor::RegisterCollectionWorkflowDelegates(FTelemetryRouter& Router)
+{
+	Router.OnTelemetry<FAssetAddedToCollectionTelemetryEvent>([](const FAssetAddedToCollectionTelemetryEvent& Event)
+	{
+		const int SchemaVersion = 1;
+		
+		FStudioTelemetry::Get().RecordEvent(TEXT("Editor.Collections.AssetsAdded"),
+		{
+			{ TEXT("SchemaVersion"), SchemaVersion},
+			{ TEXT("DurationSec"), Event.DurationSec},
+			{ TEXT("ObjectCount"), Event.NumAdded},
+			{ TEXT("Workflow"), Event.Workflow},
+			{ TEXT("CollectionShareType"), ECollectionShareType::ToString(Event.CollectionShareType)},
+		});
+	});
+
+	Router.OnTelemetry<FAssetRemovedFromCollectionTelemetryEvent>([](const FAssetRemovedFromCollectionTelemetryEvent& Event)
+	{
+		const int SchemaVersion = 1;
+		
+		FStudioTelemetry::Get().RecordEvent(TEXT("Editor.Collections.AssetsRemoved"),
+		{
+			{ TEXT("SchemaVersion"), SchemaVersion},
+			{ TEXT("DurationSec"), Event.DurationSec},
+			{ TEXT("ObjectCount"), Event.NumRemoved},
+			{ TEXT("Workflow"), Event.Workflow},
+			{ TEXT("CollectionShareType"), ECollectionShareType::ToString(Event.CollectionShareType) },
+		});
+	});
+
+	Router.OnTelemetry<FCollectionCreatedTelemetryEvent>([](const FCollectionCreatedTelemetryEvent& Event)
+	{
+		const int SchemaVersion = 1;
+
+		FStudioTelemetry::Get().RecordEvent(TEXT("Editor.Collections.CollectionCreated"),
+		{
+			{ TEXT("SchemaVersion"), SchemaVersion},
+			{ TEXT("DurationSec"), Event.DurationSec},
+			{ TEXT("CollectionShareType"), ECollectionShareType::ToString(Event.CollectionShareType)},
+		});
+	});
+
+	Router.OnTelemetry<FCollectionsDeletedTelemetryEvent>([](const FCollectionsDeletedTelemetryEvent& Event)
+	{
+		const int SchemaVersion = 1;
+
+		FStudioTelemetry::Get().RecordEvent(TEXT("Editor.Collections.CollectionDeleted"),
+		{
+			{ TEXT("SchemaVersion"), SchemaVersion},
+			{ TEXT("DurationSec"), Event.DurationSec},
+			{ TEXT("ObjectCount"), Event.CollectionsDeleted},
+		});
+	});
+}
+
+void FStudioTelemetryEditor::Initialize()
+{
+	SessionStartTime = FPlatformTime::Seconds();
+
+	TSharedPtr<FAnalyticsFlowTracker> FlowTracker = FStudioTelemetry::Get().GetFlowTracker().Pin();
+
+	if (FlowTracker.IsValid())
+	{
+		// Start Editor Boot Flow immediately
+		EditorBootFlowGuid = FlowTracker->StartFlow(TEXT("Editor Boot"));
+		FlowTracker->StartSubFlow(TEXT("Editor Boot"), EditorBootFlowGuid);
+	}
+
+	FEditorDelegates::BeginPIE.AddLambda([this](bool)
+		{
+			TSharedPtr<FAnalyticsFlowTracker> FlowTracker = FStudioTelemetry::Get().GetFlowTracker().Pin();
+
+			if (FlowTracker.IsValid())
+			{
+				// Finish the Editor Interaction Flow
+				FlowTracker->EndFlow(InteractiveEditorFlowGuid);
+
+				// Start PIE Flow
+				PIEStartTime = FPlatformTime::Seconds();
+				PIEFlowGuid = FlowTracker->StartFlow(TEXT("Play In Editor"));
+				PIEInitializeSubFlowGuid = FlowTracker->StartSubFlow(TEXT("PIE Initialize"), PIEFlowGuid);
+			}			
+		});
+
+	FEditorDelegates::PostPIEStarted.AddLambda([this](bool)
+		{
+			TSharedPtr<FAnalyticsFlowTracker> FlowTracker = FStudioTelemetry::Get().GetFlowTracker().Pin();
+
+			if (FlowTracker.IsValid())
+			{
+				// Called when PIE has finally started
+				FlowTracker->EndSubFlow(PIEInitializeSubFlowGuid);
+			}
+		});
+
+	FEditorDelegates::EndPIE.AddLambda([this](bool)
+		{
+			TSharedPtr<FAnalyticsFlowTracker> FlowTracker = FStudioTelemetry::Get().GetFlowTracker().Pin();
+
+			if (FlowTracker.IsValid())
+			{
+				// End PIE Flow and any SubFlows within it
+				FlowTracker->EndFlow(PIEFlowGuid, true);
+
+				// Restart the Interactive Editor Flow
+				FlowTracker->StartFlow(TEXT("Interactive Editor"));
+			}
+		});
+
+	FWorldDelegates::OnPIEStarted.AddLambda([this](UGameInstance* GameInstance)
+		{
+			TSharedPtr<FAnalyticsFlowTracker> FlowTracker = FStudioTelemetry::Get().GetFlowTracker().Pin();
+
+			if (FlowTracker.IsValid())
+			{
+				PIELoadMapStartTime = FPlatformTime::Seconds();
+				PIELoadMapSubFlowGuid = FlowTracker->StartSubFlow(TEXT("PIE Load Map"), PIEFlowGuid);
+			}
+		});
+
+	FWorldDelegates::OnPIEReady.AddLambda([this](UGameInstance* GameInstance)
+		{	
+			const FString MapName = FPaths::GetBaseFilename(GameInstance->PIEMapName);
+
+			TArray<FAnalyticsEventAttribute> Attributes;
+			Attributes.Emplace(FAnalyticsEventAttribute(TEXT("LevelName"), LevelName));
+			Attributes.Emplace(FAnalyticsEventAttribute(TEXT("MapName"), MapName));
+			
+			// Currently there is no obvious way to know when PIE has finished loading everything but loading map FrontEnd is good enough for Epic products.
+			if (MapName.Contains(TEXT("Frontend")))
+			{
+				static bool IsFirstTimeToPIE = true;
+
+				PIEStartupTime = FPlatformTime::Seconds() - PIEStartTime;
+
+				if (IsFirstTimeToPIE == true)
+				{
+					// Record the absolute time from editor boot to PIE
+					FStudioTelemetryEditor::RecordEvent_Loading(TEXT("TimeToPIE"), EditorStartupTime + LoadMapTime + PIEStartupTime, Attributes);
+					FStudioTelemetryEditor::RecordEvent_CoreSystems(TEXT("TimeToPIE"), Attributes);
+
+					IsFirstTimeToPIE = false;
+				}
+				
+				// Record the time from start PIE to PIE
+				FStudioTelemetryEditor::RecordEvent_Loading(TEXT("PIE.TotalStartupTime"), PIEStartupTime, Attributes);
+				FStudioTelemetryEditor::RecordEvent_CoreSystems(TEXT("PIE.TotalStartupTime"), Attributes);
+			}
+
+			PIELoadMapTime = FPlatformTime::Seconds() - PIELoadMapStartTime;
+
+			// Record the sub map load time
+			FStudioTelemetryEditor::RecordEvent_Loading(TEXT("PIE.LoadMapTime"), PIELoadMapTime, Attributes);
+			FStudioTelemetryEditor::RecordEvent_CoreSystems(TEXT("PIE.LoadMapTime"), Attributes);
+
+			TSharedPtr<FAnalyticsFlowTracker> FlowTracker = FStudioTelemetry::Get().GetFlowTracker().Pin();
+
+			if (FlowTracker.IsValid())
+			{
+				FlowTracker->EndSubFlow(PIELoadMapSubFlowGuid, true, Attributes);
+
+				// Start tracking World Streaming
+				if (UWorld* World = GameInstance->GetWorld())
+				{
+					FGuid WorldStreamingSubFlowGuid = FlowTracker->StartSubFlow(TEXT("World Streaming"));
+
+					WorldStreamingStartTime = FPlatformTime::Seconds();
+
+					World->OnWorldMatchStarting.AddLambda([this, WorldStreamingSubFlowGuid, MapName, FlowTracker]()
+						{
+							TArray<FAnalyticsEventAttribute> Attributes;
+							Attributes.Emplace(FAnalyticsEventAttribute(TEXT("LevelName"), LevelName));
+							Attributes.Emplace(FAnalyticsEventAttribute(TEXT("MapName"), MapName));
+
+							FlowTracker->EndSubFlow(WorldStreamingSubFlowGuid, true, Attributes);
+							
+							FStudioTelemetryEditor::RecordEvent_Loading(TEXT("WorldStreaming"), FPlatformTime::Seconds() - WorldStreamingStartTime, Attributes);
+						});
+				}
+			}
+		});
+
+	FEditorFileUtils::GetOnLoadMapStartDelegate().AddLambda([this]()
+		{
+			TSharedPtr<FAnalyticsFlowTracker> FlowTracker = FStudioTelemetry::Get().GetFlowTracker().Pin();
+
+			if (FlowTracker.IsValid())
+			{
+				// Load Map SubFlow will inherit the current Flow scope
+				LoadMapStartTime = FPlatformTime::Seconds();
+				LoadMapSubFlowGuid = FlowTracker->StartSubFlow(TEXT("Load Map"));
+			}
+		});
+
+	FEditorFileUtils::GetOnLoadMapEndDelegate().AddLambda([this](const FString& MapName)
+		{
+			LevelName = MapName;
+
+			TArray<FAnalyticsEventAttribute> Attributes;
+			Attributes.Emplace(FAnalyticsEventAttribute(TEXT("LevelName"), LevelName));
+			Attributes.Emplace(FAnalyticsEventAttribute(TEXT("MapName"), MapName));
+
+			LoadMapTime = FPlatformTime::Seconds() - LoadMapStartTime;
+			
+			FStudioTelemetryEditor::RecordEvent_Loading(TEXT("LoadMap"), LoadMapTime, Attributes);
+			FStudioTelemetryEditor::RecordEvent_CoreSystems(TEXT("LoadMap"), Attributes);
+
+			TSharedPtr<FAnalyticsFlowTracker> FlowTracker = FStudioTelemetry::Get().GetFlowTracker().Pin();
+
+			if (FlowTracker.IsValid())
+			{
+				FlowTracker->EndSubFlow(LoadMapSubFlowGuid, true, Attributes);
+			}
+		});
+
+	FEditorDelegates::OnEditorBoot.AddLambda([this](double TimeToBootEditor)
+		{
+			// Editor Boot Flow has finished
+			TArray<FAnalyticsEventAttribute> Attributes;
+			Attributes.Emplace(FAnalyticsEventAttribute(TEXT("LevelName"), TEXT("EditorBoot")));
+			Attributes.Emplace(FAnalyticsEventAttribute(TEXT("MapName"), TEXT("EditorBoot")));
+			
+			FStudioTelemetryEditor::RecordEvent_Loading(TEXT("BootEditor"), TimeToBootEditor, Attributes);
+			FStudioTelemetryEditor::RecordEvent_CoreSystems(TEXT("BootEditor"), Attributes);
+
+			TSharedPtr<FAnalyticsFlowTracker> FlowTracker = FStudioTelemetry::Get().GetFlowTracker().Pin();
+
+			if (FlowTracker.IsValid())
+			{
+				FlowTracker->EndFlow(EditorBootFlowGuid, true, Attributes);
+
+				// Start the Editor Interact Flow
+				InteractiveEditorFlowGuid = FlowTracker->StartFlow(TEXT("Interactive Editor"));
+			}
+		});
+
+	FEditorDelegates::OnEditorInitialized.AddLambda([this](double TimeToInitializeEditor)
+		{
+			EditorStartupTime = TimeToInitializeEditor;
+
+			// Editor has initialized
+			TArray<FAnalyticsEventAttribute> Attributes;
+			Attributes.Emplace(FAnalyticsEventAttribute(TEXT("LevelName"), TEXT("EditorInitialize")));
+			Attributes.Emplace(FAnalyticsEventAttribute(TEXT("MapName"), TEXT("EditorInitialize")));
+			
+			FStudioTelemetryEditor::RecordEvent_Loading(TEXT("TotalEditorStartup"), TimeToInitializeEditor, Attributes);
+			FStudioTelemetryEditor::RecordEvent_CoreSystems(TEXT("TotalEditorStartup"), Attributes);
+
+			ensureMsgf(GEditor, TEXT("GEditor was not valid"));
+
+			if (GEditor != nullptr)
+			{
+				// Install callbacks for Open Asset Dialogue
+				GEditor->GetEditorSubsystem<UAssetEditorSubsystem>()->OnAssetEditorRequestedOpen().AddLambda([this](UObject* Asset)
+					{
+						AssetOpenStartTime = FPlatformTime::Seconds();
+						TSharedPtr<FAnalyticsFlowTracker> FlowTracker = FStudioTelemetry::Get().GetFlowTracker().Pin();
+
+						if (FlowTracker.IsValid())
+						{
+							FlowTracker->StartSubFlow(TEXT("Open Asset Editor"));
+						}
+					});
+
+				GEditor->GetEditorSubsystem<UAssetEditorSubsystem>()->OnAssetOpenedInEditor().AddLambda([this](UObject* Asset, IAssetEditorInstance*)
+					{
+						TArray<FAnalyticsEventAttribute> Attributes;
+						Attributes.Emplace(FAnalyticsEventAttribute(TEXT("LevelName"), LevelName));
+						
+						if (Asset != nullptr)
+						{
+							Attributes.Emplace(FAnalyticsEventAttribute(TEXT("MapName"), Asset->GetFullName()));
+							Attributes.Emplace(TEXT("AssetPath"), Asset->GetFullName());
+							Attributes.Emplace(TEXT("AssetType"), Asset->GetClass()->GetName());
+
+							FStudioTelemetryEditor::RecordEvent_Loading(TEXT("OpenAssetEditor"), FPlatformTime::Seconds() - AssetOpenStartTime, Attributes);
+						}
+
+						TSharedPtr<FAnalyticsFlowTracker> FlowTracker = FStudioTelemetry::Get().GetFlowTracker().Pin();
+
+						if (FlowTracker.IsValid())
+						{
+							FlowTracker->EndSubFlow(TEXT("Open Asset Editor"), true, Attributes);
+						}
+				});
+			}
+
+			ensureMsgf(GUnrealEd, TEXT("GUnrealEd was not valid"));
+			
+			if (GUnrealEd!=nullptr && GUnrealEd->CookServer != nullptr )
+			{
+				UE::Cook::FDelegates::CookByTheBookFinished.AddLambda([this](UE::Cook::ICookInfo& CookInfo)
+					{
+						TArray<FAnalyticsEventAttribute> Attributes;
+						Attributes.Emplace(FAnalyticsEventAttribute(TEXT("LevelName"), LevelName));
+						Attributes.Emplace(FAnalyticsEventAttribute(TEXT("MapName"), LevelName));
+						
+						FStudioTelemetryEditor::RecordEvent_Cooking(Attributes);
+						FStudioTelemetryEditor::RecordEvent_CoreSystems(TEXT("Cooking"), Attributes);
+					});
+			}
+
+			ensureMsgf(GWarn, TEXT("GWarn was not valid"));
+
+			if (GWarn != nullptr)
+			{
+				GWarn->OnStartSlowTask().AddLambda([this](const FText& TaskName)
+					{	
+						TSharedPtr<FAnalyticsFlowTracker> FlowTracker = FStudioTelemetry::Get().GetFlowTracker().Pin();
+
+						if (FlowTracker.IsValid())
+						{
+							FlowTracker->StartSubFlow(FName(TaskName.ToString()));
+						}
+					});
+
+				GWarn->OnFinalizeSlowTask().AddLambda([this](const FText& TaskName, double TaskDurationSeconds)
+					{
+						TArray<FAnalyticsEventAttribute> Attributes;
+						Attributes.Emplace(FAnalyticsEventAttribute(TEXT("LevelName"), LevelName));
+						Attributes.Emplace(TEXT("TaskName"), TaskName.ToString());
+						
+						TSharedPtr<FAnalyticsFlowTracker> FlowTracker = FStudioTelemetry::Get().GetFlowTracker().Pin();
+
+						if (FlowTracker.IsValid())
+						{
+							FlowTracker->EndSubFlow(FName(TaskName.ToString()), true, Attributes);
+						}
+						
+						FStudioTelemetryEditor::RecordEvent_Loading(TEXT("SlowTaskDialog"), TaskDurationSeconds, Attributes);
+					});
+			}
+		});
+	
+	// Install Cooking Callbacks
+	UE::Cook::FDelegates::CookByTheBookStarted.AddLambda([this](UE::Cook::ICookInfo& CookInfo)
+	{
+			TSharedPtr<FAnalyticsFlowTracker> FlowTracker = FStudioTelemetry::Get().GetFlowTracker().Pin();
+
+			if (FlowTracker.IsValid())
+			{
+				CookByTheBookFlowGuid = FlowTracker->StartFlow("Cook By The Book");
+			}
+	});
+
+	UE::Cook::FDelegates::CookByTheBookFinished.AddLambda([this](UE::Cook::ICookInfo& CookInfo)
+	{
+		// Suppress sending telemetry from CookWorkers for now.
+		uint32 MultiprocessId = 0;
+		FParse::Value(FCommandLine::Get(), TEXT("-MultiprocessId="), MultiprocessId);
+		if (MultiprocessId != 0)
+		{
+			return;
+		}
+
+		TArray<FAnalyticsEventAttribute> Attributes;
+		Attributes.Emplace(FAnalyticsEventAttribute(TEXT("LevelName"), LevelName));
+		Attributes.Emplace(FAnalyticsEventAttribute(TEXT("MapName"), LevelName));
+
+		TSharedPtr<FAnalyticsFlowTracker> FlowTracker = FStudioTelemetry::Get().GetFlowTracker().Pin();
+
+		if (FlowTracker.IsValid())
+		{
+			FlowTracker->EndFlow(CookByTheBookFlowGuid, true, Attributes);
+		}
+
+		FStudioTelemetryEditor::RecordEvent_Cooking(Attributes);
+		FStudioTelemetryEditor::RecordEvent_CoreSystems(TEXT("Cooking"), Attributes);
+	});
+
+	FContentBrowserModule* ContentBrowserModule = FModuleManager::GetModulePtr<FContentBrowserModule>( TEXT("ContentBrowser") );
+	
+	FTelemetryRouter& Router = FTelemetryRouter::Get();
+	{
+		using namespace UE::Telemetry::ContentBrowser;
+		Router.OnTelemetry<FBackendFilterTelemetry>([this](const FBackendFilterTelemetry& Data)
+		{
+			FString DataFilterText = LexToString(FJsonNull{});
+			if (Data.DataFilter)
+			{
+				Private::FAnalyticsJsonWriter J(&DataFilterText);
+				J.WriteObjectStart();
+				J.WriteValue("RecursivePaths", Data.DataFilter->bRecursivePaths);
+				J.WriteValue("ItemTypeFilter", UEnum::GetValueOrBitfieldAsString(Data.DataFilter->ItemTypeFilter));
+				J.WriteValue("ItemCategoryFilter", UEnum::GetValueOrBitfieldAsString(Data.DataFilter->ItemCategoryFilter));
+				J.WriteValue("ItemAttributeFilter", UEnum::GetValueOrBitfieldAsString(Data.DataFilter->ItemAttributeFilter));
+				TArray<const UScriptStruct*> FilterTypes = Data.DataFilter->ExtraFilters.GetFilterTypes();
+				if (FilterTypes.Num() > 0)
+				{
+					J.WriteArrayStart("FilterTypes");
+					for (const UScriptStruct* Type : FilterTypes)
+					{
+						J.WriteValue(Type->GetPathName());
+					}
+					J.WriteArrayEnd();				
+				}
+				J.WriteObjectEnd();
+				J.Close();
+			}
+			
+			FStudioTelemetry::Get().RecordEvent(TEXT("Editor.AssetView.BackendFilter"), 
+			{
+				{ TEXT("SchemaVersion"), 1 },
+				{ TEXT("ViewCorrelationGuid"), Data.ViewCorrelationGuid },
+				{ TEXT("FilterSessionCorrelationGuid"), Data.FilterSessionCorrelationGuid },
+				{ TEXT("HasCustomItemSources"), Data.bHasCustomItemSources },
+				{ TEXT("RefreshSourceItemsDurationSeconds"), Data.RefreshSourceItemsDurationSeconds },
+				{ TEXT("NumBackendItems"), Data.NumBackendItems },
+				{ TEXT("DataFilter"), FJsonFragment(MoveTemp(DataFilterText)) },
+			});
+		});
+
+		Router.OnTelemetry<FFrontendFilterTelemetry>([this](const FFrontendFilterTelemetry& Data)
+		{
+			FString FilterText = LexToString(FJsonNull{});
+			if (Data.FrontendFilters.IsValid() && Data.FrontendFilters->Num())
+			{
+				Private::FAnalyticsJsonWriter J(&FilterText);
+				J.WriteArrayStart();
+				for (int32 i=0; i < Data.FrontendFilters->Num(); ++i)
+				{
+					TSharedPtr<IFilter<FAssetFilterType>> Filter = Data.FrontendFilters->GetFilterAtIndex(i);
+					J.WriteValue(Filter->GetName());
+				}
+				J.WriteArrayEnd();
+				J.Close();
+			}
+			FStudioTelemetry::Get().RecordEvent(TEXT("Editor.AssetView.FrontendFilter"), 
+			{
+				{ TEXT("SchemaVersion"), 1 },
+				{ TEXT("ViewCorrelationGuid"), Data.ViewCorrelationGuid },
+				{ TEXT("FilterSessionCorrelationGuid"), Data.FilterSessionCorrelationGuid },
+				{ TEXT("TotalItemsToFilter"), Data.TotalItemsToFilter },
+				{ TEXT("PriorityItemsToFilter"), Data.PriorityItemsToFilter },
+				{ TEXT("TotalResults"), Data.TotalResults },
+				{ TEXT("AmortizeDurationSeconds"), Data.AmortizeDuration },
+				{ TEXT("WorkDurationSeconds"), Data.WorkDuration },
+				{ TEXT("ResultLatency"), AnalyticsOptionalToStringOrNull(Data.ResultLatency) },
+				{ TEXT("TimeUntilInteractionSeconds"), AnalyticsOptionalToStringOrNull(Data.TimeUntilInteraction) },
+				{ TEXT("Completed"), Data.bCompleted },
+				{ TEXT("FrontendFilters"), FJsonFragment(MoveTemp(FilterText)) },
+			});
+		});
+
+		RegisterCollectionWorkflowDelegates(Router);
+	}
+	{
+		using namespace UE::Telemetry::AssetRegistry;
+
+		Router.OnTelemetry<FStartupTelemetry>([this](const FStartupTelemetry& Data){
+			FStudioTelemetry::Get().RecordEvent(TEXT("Editor.AssetRegistry.Startup"), 
+			{
+				{ TEXT("SchemaVersion"), 1 },
+				{ TEXT("Duration"), Data.StartupDuration },
+				{ TEXT("StartedAsyncGather"), Data.bStartedAsyncGather },
+			});
+		});
+		Router.OnTelemetry<FSynchronousScanTelemetry>([this](const FSynchronousScanTelemetry& Data){
+			if (Data.Duration < 0.5)
+			{
+				return;
+			}
+			FString DirectoriesText;
+			{
+				Private::FAnalyticsJsonWriter J(&DirectoriesText);
+				J.WriteArrayStart();
+				for (const FString& Directory : MakeArrayView(Data.Directories).Left(100))
+				{
+					J.WriteValue(Directory);
+				}
+				J.WriteArrayEnd();
+				J.Close();
+			}
+			FString FilesText;
+			{
+				Private::FAnalyticsJsonWriter J(&FilesText);
+				J.WriteArrayStart();
+				for (const FString& File : MakeArrayView(Data.Files).Left(100))
+				{
+					J.WriteValue(File);
+				}
+				J.WriteArrayEnd();
+				J.Close();
+			}
+			FStudioTelemetry::Get().RecordEvent(TEXT("Editor.AssetRegistry.SynchronousScan"), 
+			{
+				{ TEXT("SchemaVersion"), 1 },
+				{ TEXT("Directories"), FJsonFragment(MoveTemp(DirectoriesText)) },
+				{ TEXT("Files"), FJsonFragment(MoveTemp(FilesText)) },
+				{ TEXT("Flags"), LexToString(Data.Flags) },
+				{ TEXT("NumFoundAssets"), Data.NumFoundAssets },
+				{ TEXT("DurationSeconds"), Data.Duration },
+				{ TEXT("InitialSearchStarted"), Data.bInitialSearchStarted },
+				{ TEXT("InitialSearchCompleted"), Data.bInitialSearchCompleted },
+			});
+		});
+		Router.OnTelemetry<FGatherTelemetry>([this](const FGatherTelemetry& Data){
+			FStudioTelemetry::Get().RecordEvent(TEXT("Editor.AssetRegistry.InitialScan"), 
+			{
+				{ TEXT("SchemaVersion"), 1 },
+				{ TEXT("TotalDurationSeconds"), Data.TotalSearchDurationSeconds },
+				{ TEXT("TotalWorkSeconds"), Data.TotalWorkTimeSeconds },
+				{ TEXT("DiscoverySeconds"), Data.DiscoveryTimeSeconds },
+				{ TEXT("GatherSeconds"), Data.GatherTimeSeconds },
+				{ TEXT("StoreSeconds"), Data.StoreTimeSeconds },
+				{ TEXT("NumCachedDirectories"), Data.NumCachedDirectories },
+				{ TEXT("NumUncachedDirectories"), Data.NumUncachedDirectories },
+				{ TEXT("NumCachedAssetFiles"), Data.NumCachedAssetFiles },
+				{ TEXT("NumUncachedAssetFiles"), Data.NumUncachedAssetFiles },
+			});
+		});
+		Router.OnTelemetry<FDirectoryWatcherUpdateTelemetry>([this](const FDirectoryWatcherUpdateTelemetry& Data){
+			if (Data.DurationSeconds < 0.5)
+			{
+				return;
+			}
+			FStudioTelemetry::Get().RecordEvent(TEXT("Editor.AssetRegistry.DirectoryWatcherUpdate"), 
+			{
+				{ TEXT("SchemaVersion"), 1 },
+				{ TEXT("NumChanges"), Data.Changes.Num() },
+				{ TEXT("DurationSeconds"), Data.DurationSeconds },
+				{ TEXT("InitialSearchStarted"), Data.bInitialSearchStarted },
+				{ TEXT("InitialSearchCompleted"), Data.bInitialSearchCompleted },
+			});
+		});
+	}
+}
+
+void FStudioTelemetryEditor::Shutdown()
+{
+}
+
+UE_ENABLE_OPTIMIZATION_SHIP
+
+#endif // WITH_EDITOR
\ No newline at end of file
diff --git a/Engine/Plugins/StudioTelemetry/Source/Private/StudioTelemetryEditor.h b/Engine/Plugins/StudioTelemetry/Source/Private/StudioTelemetryEditor.h
new file mode 100644
--- /dev/null
+++ b/Engine/Plugins/StudioTelemetry/Source/Private/StudioTelemetryEditor.h
@@ -0,0 +1,60 @@
+// Copyright Epic Games, Inc. All Rights Reserved.
+
+#pragma once
+
+#if WITH_EDITOR
+
+#include "StudioTelemetry.h"
+
+class FTelemetryRouter;
+
+/**
+ * A class that implements a variety of pre-configured Core and Editor telemetry events that can be used to evaluate the efficiecnty of the most common developer workflows
+ */
+class FStudioTelemetryEditor : FNoncopyable
+{
+public:
+	FStudioTelemetryEditor() {};
+	~FStudioTelemetryEditor() {};
+
+	static FStudioTelemetryEditor& Get();
+	void Initialize();
+	void Shutdown();
+
+private:	
+
+	/** Various event handling functions */
+	static void RecordEvent_Cooking(TArray<FAnalyticsEventAttribute> Attributes = {});
+	static void RecordEvent_Loading(const FString& LoadingName, double LoadingSeconds, TArray<FAnalyticsEventAttribute> Attributes = {});
+	static void RecordEvent_CoreSystems(const FString& Context, TArray<FAnalyticsEventAttribute> Attributes = {});
+	static void RecordEvent_DDCResource(const FString& Context, TArray<FAnalyticsEventAttribute> Attributes = {});
+	static void RecordEvent_DDCSummary(const FString& Context, TArray<FAnalyticsEventAttribute> Attributes = {});
+	static void RecordEvent_IAS(const FString& Context, TArray<FAnalyticsEventAttribute> Attributes = {});
+	static void RecordEvent_Zen(const FString& Context, TArray<FAnalyticsEventAttribute> Attributes = {});
+	static void RecordEvent_VirtualAssets(const FString& Context, TArray<FAnalyticsEventAttribute> Attributes = {});
+	static void RegisterCollectionWorkflowDelegates(FTelemetryRouter& Router);
+
+	/** Internal variables for Flows and Timer data */
+	FString LevelName;
+	FGuid EditorBootFlowGuid;
+	FGuid InteractiveEditorFlowGuid;
+	FGuid LoadMapSubFlowGuid;
+	FGuid PIEFlowGuid;
+	FGuid PIEInitializeSubFlowGuid;
+	FGuid PIELoadMapSubFlowGuid;
+	FGuid CookByTheBookFlowGuid;
+
+	double SessionStartTime;
+	double PIEStartTime;
+	double LoadMapStartTime;
+	double AssetOpenStartTime;
+	double WorldStreamingStartTime;
+	double EditorStartupTime=0;
+	double LoadMapTime=0;
+	double PIEStartupTime = 0;
+	double PIELoadMapTime = 0;
+	double PIELoadMapStartTime;
+};
+
+#endif // WITH_EDITOR
+
diff --git a/Engine/Plugins/StudioTelemetry/Source/Private/StudioTelemetryLog.h b/Engine/Plugins/StudioTelemetry/Source/Private/StudioTelemetryLog.h
new file mode 100644
--- /dev/null
+++ b/Engine/Plugins/StudioTelemetry/Source/Private/StudioTelemetryLog.h
@@ -0,0 +1,10 @@
+// Copyright Epic Games, Inc. All Rights Reserved.
+
+#pragma once
+
+#include "CoreMinimal.h"
+
+/**
+ * StudioTelemetry plugin logging
+ */
+DECLARE_LOG_CATEGORY_EXTERN(LogStudioTelemetry, Log, All);
diff --git a/Engine/Plugins/StudioTelemetry/Source/Public/StudioTelemetry.h b/Engine/Plugins/StudioTelemetry/Source/Public/StudioTelemetry.h
new file mode 100644
--- /dev/null
+++ b/Engine/Plugins/StudioTelemetry/Source/Public/StudioTelemetry.h
@@ -0,0 +1,64 @@
+// Copyright Epic Games, Inc. All Rights Reserved.
+
+#pragma once
+
+#include "IAnalyticsProviderET.h"
+#include "AnalyticsFlowTracker.h"
+#include "Containers/Array.h"
+#include "Containers/UnrealString.h"
+#include "CoreMinimal.h"
+#include "Modules/ModuleManager.h"
+#include "Templates/SharedPointer.h"
+
+class FAnalyticsProviderMulticast;
+
+/**
+ * Public facing Studio Telemetry Plugin API
+ */
+class FStudioTelemetry : public IModuleInterface
+{
+public:
+
+	using OnEventRecorded = IAnalyticsProviderET::OnEventRecorded;
+
+	/** Check whether the module is available*/
+	static STUDIOTELEMETRY_API bool IsAvaliable() { return FModuleManager::Get().IsModuleLoaded("StudioTelemetry"); }
+
+	/** Access to the module singleton*/
+	static STUDIOTELEMETRY_API FStudioTelemetry& Get();
+
+	/** Access to the a specific named analytics provider within the system*/
+	STUDIOTELEMETRY_API TWeakPtr<IAnalyticsProvider> GetProvider(const FString& ProviderName);
+
+	/** Access to the broadcast analytics provider for the system*/
+	STUDIOTELEMETRY_API TWeakPtr<IAnalyticsProvider> GetProvider();
+
+	/** Access to the flow tracker for the system, this will ultimately broadcast flow events to all providers*/
+	STUDIOTELEMETRY_API TWeakPtr<FAnalyticsFlowTracker> GetFlowTracker();
+	
+	/** Thread safe method to record an event to all analytics providers*/
+	STUDIOTELEMETRY_API void RecordEvent(const FString& EventName, const TArray<FAnalyticsEventAttribute>& Attributes = {});
+
+	/** Thread safe method to record an event to the named analytics provider */
+	STUDIOTELEMETRY_API void RecordEvent(const FString& ProviderName, const FString& EventName, const TArray<FAnalyticsEventAttribute>& Attributes = {});
+	
+	/** Method for custom recording of telemetry events*/
+	STUDIOTELEMETRY_API void SetOnEventRecordedCallback(OnEventRecorded);
+
+private:
+
+	/** IModuleInterface implementation */
+	STUDIOTELEMETRY_API virtual void StartupModule()  final;
+	STUDIOTELEMETRY_API virtual void ShutdownModule()  final;
+
+	/** Starts a new analytics session*/
+	void StartSession();
+
+	/** Ends an existing analytics session*/
+	void EndSession();
+
+	FCriticalSection						CriticalSection;
+	TSharedPtr<FAnalyticsProviderMulticast>	AnalyticsProvider;
+	TSharedPtr<FAnalyticsFlowTracker>		AnalyticsFlowTracker;
+	OnEventRecorded							OnEventRecordedCallback;
+};
diff --git a/Engine/Plugins/StudioTelemetry/Source/StudioTelemetry.build.cs b/Engine/Plugins/StudioTelemetry/Source/StudioTelemetry.build.cs
new file mode 100644
--- /dev/null
+++ b/Engine/Plugins/StudioTelemetry/Source/StudioTelemetry.build.cs
@@ -0,0 +1,57 @@
+// Copyright Epic Games, Inc. All Rights Reserved.
+
+using System.IO;
+using UnrealBuildTool;
+
+public class StudioTelemetry : ModuleRules
+{
+	public StudioTelemetry(ReadOnlyTargetRules Target) : base(Target)
+	{
+		PCHUsage = ModuleRules.PCHUsageMode.UseExplicitOrSharedPCHs;
+		
+		PublicDependencyModuleNames.AddRange(
+			new string[] {
+				"Core",
+				"HTTP"
+			}
+		);
+
+		PrivateDependencyModuleNames.AddRange(
+			new string[] {
+				"CoreUObject",
+				"Engine",
+                "EngineSettings",
+                "BuildSettings",
+                "Analytics",
+                "AnalyticsET",
+				"TelemetryUtils",
+				"RHI"
+			}
+		);
+
+		if (Target.Configuration != UnrealTargetConfiguration.Shipping)
+		{
+			PublicDependencyModuleNames.AddRange(
+			new string[]
+			{
+				"Horde"
+			});
+		}
+
+		if (Target.bBuildEditor)
+		{
+			PrivateDependencyModuleNames.AddRange(
+				new string[] {
+					"AssetRegistry",
+					"UnrealEd",
+					"DerivedDataCache",
+					"Zen",
+					"IoStoreOnDemand",
+					"ContentBrowser",
+					"ContentBrowserData",
+					"TelemetryUtils"
+				}
+			);
+		}
+	}
+}
diff --git a/Engine/Plugins/StudioTelemetry/StudioTelemetry.uplugin b/Engine/Plugins/StudioTelemetry/StudioTelemetry.uplugin
new file mode 100644
--- /dev/null
+++ b/Engine/Plugins/StudioTelemetry/StudioTelemetry.uplugin
@@ -0,0 +1,24 @@
+{
+	"FileVersion": 1,
+	"Version": 1,
+	"VersionName": "1.0",
+	"FriendlyName": "Studio Telemetry",
+	"Description": "Centralised plugin for Studio telemetry data",
+	"Category": "Telemetry",
+	"CreatedBy": "Epic Games, Inc.",
+	"CreatedByURL": "https://epicgames.com",
+	"DocsURL": "",
+	"MarketplaceURL": "",
+	"SupportURL": "",
+	"EnabledByDefault": true,
+	"CanContainContent": false,
+	"IsBetaVersion": false,
+	"Installed": false,
+	"Modules": [
+		{
+			"Name": "StudioTelemetry",
+			"Type": "EditorAndProgram",
+			"LoadingPhase": "PostConfigInit"
+		}
+	]
+}
\ No newline at end of file
diff --git a/Engine/Source/Developer/DerivedDataCache/Private/DerivedDataCache.cpp b/Engine/Source/Developer/DerivedDataCache/Private/DerivedDataCache.cpp
--- a/Engine/Source/Developer/DerivedDataCache/Private/DerivedDataCache.cpp
+++ b/Engine/Source/Developer/DerivedDataCache/Private/DerivedDataCache.cpp
@@ -950,16 +950,6 @@ public:
 
 		GatherResourceStats(ResourceStats);
 
-		FDerivedDataCacheResourceStat ResourceStatsTotal(TEXT("Total"));
-
-		// Accumulate Totals
-		for (const FDerivedDataCacheResourceStat& Stat : ResourceStats)
-		{
-			ResourceStatsTotal += Stat;
-		}
-
-		ResourceStats.Emplace(ResourceStatsTotal);
-
 		// Append to the attributes
 		for (const FDerivedDataCacheResourceStat& Stat : ResourceStats)
 		{
@@ -1018,12 +1008,15 @@ public:
 			}
 		}
 
+		// Gather the backend stats
 		TSharedRef<FDerivedDataCacheStatsNode> RootNode = Backend->GatherUsageStats();
 		RootNode->ForEachDescendant([&Attributes](TSharedRef<const FDerivedDataCacheStatsNode> Node)
 			{
+				const FString& CacheName = Node->GetCacheName();
+
 				for (const FCookStatsManager::StringKeyValue& Stat : Node->CustomStats)
 				{
-					FString FormattedAttrName = Stat.Key.Replace(TEXT("."), TEXT("_"));
+					FString FormattedAttrName = CacheName + TEXT("_") + Stat.Key.Replace(TEXT("."), TEXT("_"));
 
 					if (Stat.Value.IsNumeric())
 					{
@@ -1033,8 +1026,6 @@ public:
 					{
 						Attributes.Emplace(FormattedAttrName, Stat.Value);
 					}
-
-
 				}
 			});
 #endif
diff --git a/Engine/Source/Developer/DerivedDataCache/Private/DerivedDataCacheUsageStats.cpp b/Engine/Source/Developer/DerivedDataCache/Private/DerivedDataCacheUsageStats.cpp
--- a/Engine/Source/Developer/DerivedDataCache/Private/DerivedDataCacheUsageStats.cpp
+++ b/Engine/Source/Developer/DerivedDataCache/Private/DerivedDataCacheUsageStats.cpp
@@ -99,6 +99,16 @@ void GatherDerivedDataCacheResourceStats(TArray<FDerivedDataCacheResourceStat>&
 	{
 		DDCResourceStats.Add(Stat);
 	}
+
+	// Accumulate Totals
+	FDerivedDataCacheResourceStat StatTotal(TEXT("Total"));
+
+	for (const FDerivedDataCacheResourceStat& Stat : DDCResourceStats)
+	{
+		StatTotal += Stat;
+	}
+
+	DDCResourceStats.Emplace(StatTotal);
 #endif
 }
 
diff --git a/Engine/Source/Developer/DerivedDataCache/Public/DerivedDataCacheInterface.h b/Engine/Source/Developer/DerivedDataCache/Public/DerivedDataCacheInterface.h
--- a/Engine/Source/Developer/DerivedDataCache/Public/DerivedDataCacheInterface.h
+++ b/Engine/Source/Developer/DerivedDataCache/Public/DerivedDataCacheInterface.h
@@ -278,6 +278,7 @@ public:
 
 	virtual void GatherSummaryStats(FDerivedDataCacheSummaryStats& DDCSummaryStats) const = 0;
 
+	UE_DEPRECATED(5.4, "This is no longer used. Use GatherUsageStats/GatherResourceStats/GatherSummaryStats instead")
 	virtual void GatherAnalytics(TArray<FAnalyticsEventAttribute>& Attributes) const = 0;
 
 	//-----------------------
diff --git a/Engine/Source/Developer/Horde/Private/Horde.cpp b/Engine/Source/Developer/Horde/Private/Horde.cpp
--- a/Engine/Source/Developer/Horde/Private/Horde.cpp
+++ b/Engine/Source/Developer/Horde/Private/Horde.cpp
@@ -4,8 +4,6 @@
 #include "Misc/CommandLine.h"
 #include "Misc/Parse.h"
 
-#if WITH_EDITOR
-
 FString FHorde::GetTemplateName()
 {
 	static FString TemplateName;
@@ -100,6 +98,4 @@ FString FHorde::GetBatchId()
 		}
 	}
 	return BatchId;
-} 
-
-#endif
\ No newline at end of file
+} 
\ No newline at end of file
diff --git a/Engine/Source/Developer/Horde/Public/Horde.h b/Engine/Source/Developer/Horde/Public/Horde.h
--- a/Engine/Source/Developer/Horde/Public/Horde.h
+++ b/Engine/Source/Developer/Horde/Public/Horde.h
@@ -6,8 +6,6 @@
 #define HORDE_API
 #endif
 
-#if WITH_EDITOR
-
 #include "Containers/UnrealString.h"
 
 struct FHorde
@@ -22,4 +20,3 @@ struct FHorde
 	HORDE_API static FString GetBatchId();
 };
 
-#endif
diff --git a/Engine/Source/Runtime/Analytics/AnalyticsET/Private/AnalyticsET.cpp b/Engine/Source/Runtime/Analytics/AnalyticsET/Private/AnalyticsET.cpp
--- a/Engine/Source/Runtime/Analytics/AnalyticsET/Private/AnalyticsET.cpp
+++ b/Engine/Source/Runtime/Analytics/AnalyticsET/Private/AnalyticsET.cpp
@@ -7,6 +7,7 @@
 #include "Analytics.h"
 #include "AnalyticsPerfTracker.h"
 #include "AnalyticsProviderETEventCache.h"
+#include "Misc/EngineVersion.h"
 
 IMPLEMENT_MODULE( FAnalyticsET, AnalyticsET );
 
@@ -25,14 +26,26 @@ void FAnalyticsET::ShutdownModule()
 }
 
 TSharedPtr<IAnalyticsProvider> FAnalyticsET::CreateAnalyticsProvider(const FAnalyticsProviderConfigurationDelegate& GetConfigValue) const
+{
+	return CreateAnalyticsProviderET(GetConfigValue);
+}
+
+TSharedPtr<IAnalyticsProviderET> FAnalyticsET::CreateAnalyticsProviderET(const FAnalyticsProviderConfigurationDelegate& GetConfigValue) const
 {
 	if (GetConfigValue.IsBound())
 	{
 		Config ConfigValues;
 		ConfigValues.APIKeyET = GetConfigValue.Execute(Config::GetKeyNameForAPIKey(), true);
 		ConfigValues.APIServerET = GetConfigValue.Execute(Config::GetKeyNameForAPIServer(), true);
+		ConfigValues.APIEndpointET = GetConfigValue.Execute(Config::GetKeyNameForAPIEndpoint(), true);
 		ConfigValues.AppVersionET = GetConfigValue.Execute(Config::GetKeyNameForAppVersion(), false);
 		ConfigValues.UseLegacyProtocol = FCString::ToBool(*GetConfigValue.Execute(Config::GetKeyNameForUseLegacyProtocol(), false));
+
+		if (ConfigValues.AppVersionET.IsEmpty())
+		{
+			ConfigValues.AppVersionET = FEngineVersion::Current().ToString();
+		}
+
 		if (!ConfigValues.UseLegacyProtocol)
 		{
 			ConfigValues.AppEnvironment = GetConfigValue.Execute(Config::GetKeyNameForAppEnvironment(), true);
@@ -46,4 +59,3 @@ TSharedPtr<IAnalyticsProvider> FAnalyticsET::CreateAnalyticsProvider(const FAnal
 	}
 	return NULL;
 }
-
diff --git a/Engine/Source/Runtime/Analytics/AnalyticsET/Private/AnalyticsFlowTracker.cpp b/Engine/Source/Runtime/Analytics/AnalyticsET/Private/AnalyticsFlowTracker.cpp
--- a/Engine/Source/Runtime/Analytics/AnalyticsET/Private/AnalyticsFlowTracker.cpp
+++ b/Engine/Source/Runtime/Analytics/AnalyticsET/Private/AnalyticsFlowTracker.cpp
@@ -3,7 +3,7 @@
 #include "AnalyticsFlowTracker.h"
 #include "ProfilingDebugging/MiscTrace.h"
 
-void FAnalyticsFlowTracker::SetProvider(TSharedPtr<IAnalyticsProviderET>& InProvider)
+void FAnalyticsFlowTracker::SetProvider(TSharedPtr<IAnalyticsProvider> InProvider)
 {
 	AnalyticsProvider = InProvider;
 }
@@ -286,9 +286,7 @@ bool FAnalyticsFlowTracker::EndFlowInternal(const FGuid& FlowGuid, bool bSuccess
 		{
 			// Aggregate the additional attributes from the sub flows
 			AggregateAttributes(EventAttributes, SubFlowData->AdditionalEventAttributes);
-
 			TotalTimeInSeconds += SubFlowData->TimeInSeconds;
-			EventAttributes.Add(FAnalyticsEventAttribute(SubFlowData->SubFlowName.ToString(), SubFlowData->TimeInSeconds));
 		}
 	}
 
diff --git a/Engine/Source/Runtime/Analytics/AnalyticsET/Private/IAnalyticsProviderET.cpp b/Engine/Source/Runtime/Analytics/AnalyticsET/Private/IAnalyticsProviderET.cpp
--- a/Engine/Source/Runtime/Analytics/AnalyticsET/Private/IAnalyticsProviderET.cpp
+++ b/Engine/Source/Runtime/Analytics/AnalyticsET/Private/IAnalyticsProviderET.cpp
@@ -283,7 +283,14 @@ FAnalyticsProviderET::FAnalyticsProviderET(const FAnalyticsET::Config& ConfigVal
 		? FString(FApp::GetBuildVersion())
 		: ConfigAppVersion.Replace(TEXT("%VERSION%"), FApp::GetBuildVersion(), ESearchCase::CaseSensitive);
 
-	UE_LOG(LogAnalytics, Display, TEXT("[%s] APIServer = %s. AppVersion = %s"), *Config.APIKeyET, *Config.APIServerET, *Config.AppVersionET);
+	if ( Config.APIEndpointET.IsEmpty() )
+	{
+		// Set a default API Endpoint
+		Config.APIEndpointET = TEXT("datarouter/api/v1/public/data");
+	}
+
+	UE_LOG(LogAnalytics, Display, TEXT("[%s] APIServer = %s%s. AppVersion = %s"), *Config.APIKeyET, *Config.APIServerET, *Config.APIEndpointET, *Config.AppVersionET);
+
 	if (Config.APIServerET.IsEmpty())
 	{
 		UE_LOG(LogAnalytics, Warning, TEXT("AnalyticsET: APIServerET is empty for APIKey (%s), creating as a NULL provider!"), *Config.APIKeyET);
@@ -359,9 +366,10 @@ PRAGMA_ENABLE_DEPRECATION_WARNINGS
 
 FAnalyticsProviderET::~FAnalyticsProviderET()
 {
-	UE_LOG(LogAnalytics, Verbose, TEXT("[%s] Destroying ET Analytics provider"), *Config.APIKeyET);
+	UE_LOG(LogAnalytics, Display, TEXT("[%s] Destroying ET Analytics provider"), *Config.APIKeyET);
 	bInDestructor = true;
 	EndSession();
+	UE_LOG(LogAnalytics, Display, TEXT("[%s] Destroyed ET Analytics provider"), *Config.APIKeyET);
 }
 
 bool FAnalyticsProviderET::StartSession(FString InSessionID, const TArray<FAnalyticsEventAttribute>& Attributes)
@@ -396,10 +404,17 @@ void FAnalyticsProviderET::EndSession()
 	SessionID.Empty();
 
 	bSessionInProgress = false;
+
+	UE_LOG(LogAnalytics, Display, TEXT("[%s] Ended ET Analytics provider session"), *Config.APIKeyET);
 }
 
 TSharedRef<IHttpRequest, ESPMode::ThreadSafe> FAnalyticsProviderET::CreateRequest()
 {
+	if (!ensure(FModuleManager::Get().IsModuleLoaded("HTTP")))
+	{
+		UE_LOG(LogAnalytics, Display, TEXT("[%s] ET Analytics provider tried to create a new HTTP request when HTTP was shutdown"), *Config.APIKeyET);
+	}
+
 	// TODO add config values for retries, for now, using default
 	TSharedRef<IHttpRequest, ESPMode::ThreadSafe> HttpRequest = HttpRetryManager->CreateRequest(FHttpRetrySystem::FRetryLimitCountSetting(),
 		FHttpRetrySystem::FRetryTimeoutRelativeSecondsSetting(),
@@ -443,7 +458,8 @@ void FAnalyticsProviderET::FlushEventsOnce()
 	{
 		TArray<uint8> Payload = EventCache.FlushCacheUTF8();
 		// UrlEncode NOTE: need to concatenate everything
-		FString URLPath  = TEXT("datarouter/api/v1/public/data?SessionID=") + FPlatformHttp::UrlEncode(SessionID);
+		FString URLPath = Config.APIEndpointET;
+				URLPath += TEXT("?SessionID=") + FPlatformHttp::UrlEncode(SessionID);
 				URLPath += TEXT("&AppID=") + FPlatformHttp::UrlEncode(Config.APIKeyET);
 				URLPath += TEXT("&AppVersion=") + FPlatformHttp::UrlEncode(Config.AppVersionET);
 				URLPath += TEXT("&UserID=") + FPlatformHttp::UrlEncode(UserID);
@@ -458,8 +474,9 @@ void FAnalyticsProviderET::FlushEventsOnce()
 			Payload.Add(TEXT('\0'));
 			// Recreate the URLPath for logging because we do not want to escape the parameters when logging.
 			// We cannot simply UrlEncode the entire Path after logging it because UrlEncode(Params) != UrlEncode(Param1) & UrlEncode(Param2) ...
-			FString LogString = FString::Printf(TEXT("[%s] AnalyticsET URL:datarouter/api/v1/public/data?SessionID=%s&AppID=%s&AppVersion=%s&UserID=%s&AppEnvironment=%s&UploadType=%s. Payload:%s"),
+			FString LogString = FString::Printf(TEXT("[%s] AnalyticsET URL:%s?SessionID=%s&AppID=%s&AppVersion=%s&UserID=%s&AppEnvironment=%s&UploadType=%s. Payload:%s"),
 				*Config.APIKeyET,
+				*Config.APIEndpointET,
 				*SessionID,
 				*Config.APIKeyET,
 				*Config.AppVersionET,
@@ -750,7 +767,7 @@ void FAnalyticsProviderET::FlushEventLegacy(const FString& EventName, const TArr
 		HttpRequest->SetHeader(TEXT("Content-Type"), TEXT("text/plain"));
 
 		// Don't need to URL encode the APIServer or the EventParams, which are already encoded, and contain parameter separaters that we DON'T want encoded.
-		FString URLPath = Config.APIServerET;
+		FString URLPath = Config.APIServerET+Config.APIEndpointET;
 		URLPath += TEXT("SendEvent.1?SessionID=") + FPlatformHttp::UrlEncode(SessionID);
 		URLPath += TEXT("&AppID=") + FPlatformHttp::UrlEncode(Config.APIKeyET);
 		URLPath += TEXT("&AppVersion=") + FPlatformHttp::UrlEncode(Config.AppVersionET);
diff --git a/Engine/Source/Runtime/Analytics/AnalyticsET/Public/AnalyticsET.h b/Engine/Source/Runtime/Analytics/AnalyticsET/Public/AnalyticsET.h
--- a/Engine/Source/Runtime/Analytics/AnalyticsET/Public/AnalyticsET.h
+++ b/Engine/Source/Runtime/Analytics/AnalyticsET/Public/AnalyticsET.h
@@ -51,6 +51,8 @@ public:
 		FString APIKeyET;
 		/** ET API Server - Base URL to send events. Set this to an empty string to essentially create a NULL analytics provider that will be non-null, but won't actually send events. */
 		FString APIServerET;
+		/** ET API Endpoint - This is the API endpoint for the provider. */
+		FString APIEndpointET;
 		/** ET Alt API Servers - Base URLs to send events on retry. */
 		TArray<FString> AltAPIServersET;
 		/** 
@@ -96,6 +98,8 @@ public:
 		static FString GetKeyNameForAPIKey() { return TEXT("APIKeyET"); }
 		/** KeyName required for APIServer configuration. */
 		static FString GetKeyNameForAPIServer() { return TEXT("APIServerET"); }
+		/** KeyName required for APIEndpoint configuration. */
+		static FString GetKeyNameForAPIEndpoint() { return TEXT("APIEndpointET"); }
 		/** KeyName required for AppVersion configuration. */
 		static FString GetKeyNameForAppVersion() { return TEXT("AppVersionET"); }
 		/** Optional parameter to use the legacy backend protocol. */
@@ -120,6 +124,13 @@ public:
 	 * The keys required exactly match the field names in the Config object. 
 	 */
 	ANALYTICSET_API virtual TSharedPtr<IAnalyticsProvider> CreateAnalyticsProvider(const FAnalyticsProviderConfigurationDelegate& GetConfigValue) const override;
+
+	/**
+	 * Construct an ET analytics provider given a configuration delegate.
+	 * The keys required exactly match the field names in the Config object.
+	 */
+	ANALYTICSET_API virtual TSharedPtr<IAnalyticsProviderET> CreateAnalyticsProviderET(const FAnalyticsProviderConfigurationDelegate& GetConfigValue) const;
+
 	
 	/** 
 	 * Construct an ET analytics provider directly from a config object.
diff --git a/Engine/Source/Runtime/Analytics/AnalyticsET/Public/AnalyticsFlowTracker.h b/Engine/Source/Runtime/Analytics/AnalyticsET/Public/AnalyticsFlowTracker.h
--- a/Engine/Source/Runtime/Analytics/AnalyticsET/Public/AnalyticsFlowTracker.h
+++ b/Engine/Source/Runtime/Analytics/AnalyticsET/Public/AnalyticsFlowTracker.h
@@ -15,7 +15,7 @@ public:
 	~FAnalyticsFlowTracker() {};
 
 	/** Sets the analytics provider for the flow tracker. */
-	ANALYTICSET_API void SetProvider(TSharedPtr<IAnalyticsProviderET>& AnalyticsProvider);
+	ANALYTICSET_API void SetProvider(TSharedPtr<IAnalyticsProvider> AnalyticsProvider);
 
 	/** Begins a new flow tracking session. Will emit Flow and SubFlow events to the specified analytics provider */
 	ANALYTICSET_API void StartSession();
@@ -88,7 +88,7 @@ private:
 	TMap<FGuid, FSubFlowData>		SubFlowDataRegistry;
 	TArray<FGuid>					FlowDataStack;
 	FCriticalSection				CriticalSection;
-	TSharedPtr<IAnalyticsProviderET> AnalyticsProvider;
+	TSharedPtr<IAnalyticsProvider>	AnalyticsProvider;
 
 	const uint32 FlowSchemaVersion = 4;
 	const FString FlowEventName = TEXT("Iteration.Flow");
diff --git a/Engine/Source/Runtime/Engine/Private/Analytics/StudioAnalytics.cpp b/Engine/Source/Runtime/Engine/Private/Analytics/StudioAnalytics.cpp
--- a/Engine/Source/Runtime/Engine/Private/Analytics/StudioAnalytics.cpp
+++ b/Engine/Source/Runtime/Engine/Private/Analytics/StudioAnalytics.cpp
@@ -132,55 +132,3 @@ void FStudioAnalytics::RecordEvent(const FString& EventName, const TArray<FAnaly
 		FStudioAnalytics::GetProvider().RecordEvent(EventName, Attributes);
 	}
 }
-
-void FStudioAnalytics::FireEvent_Loading(const FString& LoadingName, double SecondsSpentLoading, const TArray<FAnalyticsEventAttribute>& InAttributes)
-{
-	const int SchemaVersion = 2;
-
-	// Ignore anything less than a 1/4th a second.
-	if (SecondsSpentLoading < 0.250)
-	{
-		return;
-	}
-
-	// Throw out anything over 10 hours - 
-	if (SecondsSpentLoading > 36000)
-	{
-		return;
-	}
-
-	TArray<FAnalyticsEventAttribute> Attributes;
-
-	Attributes.Emplace(TEXT("SchemaVersion"), SchemaVersion);
-	Attributes.Emplace(TEXT("LoadingName"), LoadingName);
-	Attributes.Emplace(TEXT("LoadingSeconds"), SecondsSpentLoading);
-	Attributes.Append(InAttributes);
-
-	if (FStudioAnalytics::IsAvailable())
-	{
-		FStudioAnalytics::GetProvider().RecordEvent(TEXT("Performance.Loading"), Attributes);
-	}
-
-#if ENABLE_COOK_STATS
-
-	// Gather DDC analytics
-	GetDerivedDataCacheRef().GatherAnalytics(Attributes);
-
-	// Gather Virtualization analytics
-	UE::Virtualization::IVirtualizationSystem::Get().GatherAnalytics(Attributes);
-
-#if UE_WITH_ZEN
-	// Gather Zen analytics
-	if (UE::Zen::IsDefaultServicePresent())
-	{
-		UE::Zen::GetDefaultServiceInstance().GatherAnalytics(Attributes);
-	}
-#endif
-
-	if (FStudioAnalytics::IsAvailable())
-	{
-		// Store it all in the loading event
-		FStudioAnalytics::GetProvider().RecordEvent(TEXT("Core.Loading"), Attributes);
-	}
-#endif
-}
diff --git a/Engine/Source/Runtime/Engine/Public/StudioAnalytics.h b/Engine/Source/Runtime/Engine/Public/StudioAnalytics.h
--- a/Engine/Source/Runtime/Engine/Public/StudioAnalytics.h
+++ b/Engine/Source/Runtime/Engine/Public/StudioAnalytics.h
@@ -47,10 +47,6 @@ public:
 	static ENGINE_API void RecordEvent(const FString& EventName);
 	static ENGINE_API void RecordEvent(const FString& EventName, const TArray<FAnalyticsEventAttribute>& Attributes);
 
-	/** An event for reporting load time that blocks the editor. */
-	UE_DEPRECATED(5.3, "Functionality no longer used")
-	static ENGINE_API void FireEvent_Loading(const FString& LoadingName, double SecondsSpentLoading, const TArray<FAnalyticsEventAttribute>& Attributes = TArray<FAnalyticsEventAttribute>());
-
 private:
 	static void RunTimer_Concurrent();
 
diff --git a/Engine/Source/Runtime/Experimental/IoStoreOnDemand/Private/IoStoreOnDemand.cpp b/Engine/Source/Runtime/Experimental/IoStoreOnDemand/Private/IoStoreOnDemand.cpp
--- a/Engine/Source/Runtime/Experimental/IoStoreOnDemand/Private/IoStoreOnDemand.cpp
+++ b/Engine/Source/Runtime/Experimental/IoStoreOnDemand/Private/IoStoreOnDemand.cpp
@@ -2121,6 +2121,11 @@ void FIoStoreOnDemandModule::AbandonCache()
 	}
 }
 
+bool FIoStoreOnDemandModule::IsEnabled() const
+{
+	return Backend.IsValid()? Backend->IsEnabled():DeferredAbandonCache.IsSet();
+}
+
 void FIoStoreOnDemandModule::ReportAnalytics(TArray<FAnalyticsEventAttribute>& OutAnalyticsArray) const
 {
 	if (Backend.IsValid())
diff --git a/Engine/Source/Runtime/Experimental/IoStoreOnDemand/Private/OnDemandIoDispatcherBackend.cpp b/Engine/Source/Runtime/Experimental/IoStoreOnDemand/Private/OnDemandIoDispatcherBackend.cpp
--- a/Engine/Source/Runtime/Experimental/IoStoreOnDemand/Private/OnDemandIoDispatcherBackend.cpp
+++ b/Engine/Source/Runtime/Experimental/IoStoreOnDemand/Private/OnDemandIoDispatcherBackend.cpp
@@ -1555,6 +1555,7 @@ public:
 	virtual void Mount(const FOnDemandEndpoint& Endpoint) override;
 	virtual void SetBulkOptionalEnabled(bool bEnabled) override;
 	virtual void SetEnabled(bool bEnabled) override;
+	virtual bool IsEnabled() const override;
 	virtual void AbandonCache() override;
 	virtual void ReportAnalytics(TArray<FAnalyticsEventAttribute>& OutAnalyticsArray) const override;
 
@@ -2304,6 +2305,11 @@ void FOnDemandIoBackend::SetEnabled(bool bEnabled)
 	BackendStatus.SetHttpEnabled(bEnabled);
 }
 
+bool FOnDemandIoBackend::IsEnabled() const
+{
+	return BackendStatus.IsHttpEnabled();
+}
+
 void FOnDemandIoBackend::AbandonCache()
 {
 	BackendStatus.SetCacheEnabled(false);
diff --git a/Engine/Source/Runtime/Experimental/IoStoreOnDemand/Private/OnDemandIoDispatcherBackend.h b/Engine/Source/Runtime/Experimental/IoStoreOnDemand/Private/OnDemandIoDispatcherBackend.h
--- a/Engine/Source/Runtime/Experimental/IoStoreOnDemand/Private/OnDemandIoDispatcherBackend.h
+++ b/Engine/Source/Runtime/Experimental/IoStoreOnDemand/Private/OnDemandIoDispatcherBackend.h
@@ -41,6 +41,7 @@ public:
 	virtual void Mount(const FOnDemandEndpoint& Endpoint) = 0;
 	virtual void SetBulkOptionalEnabled(bool bInEnabled) = 0;
 	virtual void SetEnabled(bool bInEnabled) = 0;
+	virtual bool IsEnabled() const = 0;
 	virtual void AbandonCache() = 0;
 	virtual void ReportAnalytics(TArray<FAnalyticsEventAttribute>& OutAnalyticsArray) const = 0;
 };
diff --git a/Engine/Source/Runtime/Experimental/IoStoreOnDemand/Public/IO/IoStoreOnDemand.h b/Engine/Source/Runtime/Experimental/IoStoreOnDemand/Public/IO/IoStoreOnDemand.h
--- a/Engine/Source/Runtime/Experimental/IoStoreOnDemand/Public/IO/IoStoreOnDemand.h
+++ b/Engine/Source/Runtime/Experimental/IoStoreOnDemand/Public/IO/IoStoreOnDemand.h
@@ -310,6 +310,7 @@ private:
 public:
 	UE_API void SetBulkOptionalEnabled(bool bInEnabled);
 	UE_API void SetEnabled(bool bInEnabled);
+	UE_API bool IsEnabled() const;
 	UE_API void AbandonCache();
 
 	UE_API void ReportAnalytics(TArray<FAnalyticsEventAttribute>& OutAnalyticsArray) const;
