Fix to remove Horde from non-editor platforms

diff --git a/Engine/Plugins/StudioTelemetry/Source/Private/StudioTelemetry.cpp b/Engine/Plugins/StudioTelemetry/Source/Private/StudioTelemetry.cpp
--- a/Engine/Plugins/StudioTelemetry/Source/Private/StudioTelemetry.cpp
+++ b/Engine/Plugins/StudioTelemetry/Source/Private/StudioTelemetry.cpp
@@ -5,9 +5,6 @@
 
 #if WITH_EDITOR
 #include "StudioTelemetryEditor.h"
-#endif
-
-#if !UE_BUILD_SHIPPING
 #include "Horde.h"
 #endif
 
@@ -137,7 +134,7 @@ void FStudioTelemetry::StartSession()
 		DefaultEventAttributes.Emplace(TEXT("Config_IsBuildMachine"), GIsBuildMachine);
 		DefaultEventAttributes.Emplace(TEXT("Config_IsRunningCommandlet"), IsRunningCommandlet());
 
-#if !UE_BUILD_SHIPPING
+#if WITH_EDITOR
 		DefaultEventAttributes.Emplace(TEXT("Horde_TemplateID"), FHorde::GetTemplateId());
 		DefaultEventAttributes.Emplace(TEXT("Horde_TemplateName"), FHorde::GetTemplateName());
 		DefaultEventAttributes.Emplace(TEXT("Horde_JobURL"), FHorde::GetJobURL());
diff --git a/Engine/Plugins/StudioTelemetry/Source/StudioTelemetry.build.cs b/Engine/Plugins/StudioTelemetry/Source/StudioTelemetry.build.cs
--- a/Engine/Plugins/StudioTelemetry/Source/StudioTelemetry.build.cs
+++ b/Engine/Plugins/StudioTelemetry/Source/StudioTelemetry.build.cs
@@ -29,19 +29,11 @@ public class StudioTelemetry : ModuleRules
 			}
 		);
 
-		if (Target.Configuration != UnrealTargetConfiguration.Shipping)
-		{
-			PublicDependencyModuleNames.AddRange(
-			new string[]
-			{
-				"Horde"
-			});
-		}
-
 		if (Target.bBuildEditor)
 		{
 			PrivateDependencyModuleNames.AddRange(
 				new string[] {
+					"Horde",
 					"AssetRegistry",
 					"UnrealEd",
 					"DerivedDataCache",
