Normalizing Build_BranchName to be lower case.
Some instances of telemetry have seen to have varying case and this confuses Grafana dashboards

diff --git a/Engine/Plugins/StudioTelemetry/Source/Private/StudioTelemetry.cpp b/Engine/Plugins/StudioTelemetry/Source/Private/StudioTelemetry.cpp
--- a/Engine/Plugins/StudioTelemetry/Source/Private/StudioTelemetry.cpp
+++ b/Engine/Plugins/StudioTelemetry/Source/Private/StudioTelemetry.cpp
@@ -120,7 +120,7 @@ void FStudioTelemetry::StartSession()
 		DefaultEventAttributes.Emplace(TEXT("Build_IsInternalBuild"), FEngineBuildSettings::IsInternalBuild());
 		DefaultEventAttributes.Emplace(TEXT("Build_IsPerforceBuild"), FEngineBuildSettings::IsPerforceBuild());
 		DefaultEventAttributes.Emplace(TEXT("Build_IsPromotedBuild"), FApp::GetEngineIsPromotedBuild() == 0 ? false : true);
-		DefaultEventAttributes.Emplace(TEXT("Build_BranchName"), FApp::GetBranchName());
+		DefaultEventAttributes.Emplace(TEXT("Build_BranchName"), FApp::GetBranchName().ToLower());
 		DefaultEventAttributes.Emplace(TEXT("Build_Changelist"), BuildSettings::GetCurrentChangelist());
 
 		DefaultEventAttributes.Emplace(TEXT("Hardware_GPU"), GRHIAdapterName);
