Minor tweak to telemetry to avoid sending empty Asset records
Bumped loading, cooking and DDC schema to V3

diff --git a/Engine/Plugins/StudioTelemetry/Source/Private/StudioTelemetryEditor.cpp b/Engine/Plugins/StudioTelemetry/Source/Private/StudioTelemetryEditor.cpp
--- a/Engine/Plugins/StudioTelemetry/Source/Private/StudioTelemetryEditor.cpp
+++ b/Engine/Plugins/StudioTelemetry/Source/Private/StudioTelemetryEditor.cpp
@@ -80,7 +80,7 @@ void FStudioTelemetryEditor::RecordEvent_Cooking(TArray<FAnalyticsEventAttribute
 {
 #if ENABLE_COOK_STATS
 
-	const int SchemaVersion = 2;
+	const int SchemaVersion = 3;
 
 	Attributes.Emplace(TEXT("SchemaVersion"), SchemaVersion);
 
@@ -151,7 +151,7 @@ void FStudioTelemetryEditor::RecordEvent_Cooking(TArray<FAnalyticsEventAttribute
 
 void FStudioTelemetryEditor::RecordEvent_Loading(const FString& LoadingName, double LoadingSeconds, TArray<FAnalyticsEventAttribute> Attributes )
 {
-	const int SchemaVersion = 2;
+	const int SchemaVersion = 3;
 
 	Attributes.Emplace(TEXT("SchemaVersion"), SchemaVersion);
 	Attributes.Emplace(TEXT("LoadingName"), LoadingName);
@@ -205,7 +205,7 @@ void FStudioTelemetryEditor::RecordEvent_DDCResource(const FString& Context, TAr
 
 	GatherDerivedDataCacheResourceStats(ResourceStats);
 
-	const int SchemaVersion = 2;
+	const int SchemaVersion = 3;
 
 	Attributes.Emplace(TEXT("SchemaVersion"), SchemaVersion);
 	Attributes.Emplace(TEXT("LoadingName"), Context);
@@ -214,7 +214,13 @@ void FStudioTelemetryEditor::RecordEvent_DDCResource(const FString& Context, TAr
 	for (const FDerivedDataCacheResourceStat& Stat : ResourceStats)
 	{
 		const int64 TotalCount = Stat.BuildCount + Stat.LoadCount;
-		
+
+		if (Stat.AssetType.IsEmpty() || TotalCount==0)
+		{
+			// Empty asset type or nothing was built or loaded for this type
+			continue;
+		}
+	
 		TArray<FAnalyticsEventAttribute> EventAttributes = Attributes;
 
 		EventAttributes.Emplace(TEXT("Asset_Type"), Stat.AssetType);
@@ -224,7 +230,7 @@ void FStudioTelemetryEditor::RecordEvent_DDCResource(const FString& Context, TAr
 		EventAttributes.Emplace(TEXT("Build_Count"), Stat.BuildCount);
 		EventAttributes.Emplace(TEXT("Build_TimeSec"), Stat.BuildTimeSec);
 		EventAttributes.Emplace(TEXT("Build_SizeMB"), Stat.BuildSizeMB);
-		EventAttributes.Emplace(TEXT("Efficiency"), TotalCount>0? double(Stat.LoadCount)/double(TotalCount): 0.0 );
+		EventAttributes.Emplace(TEXT("Efficiency"), double(Stat.LoadCount)/double(TotalCount) );
 		EventAttributes.Emplace(TEXT("Thread_TimeSec"), Stat.GameThreadTimeSec);
 
 		FStudioTelemetry::Get().RecordEvent(TEXT("Core.DDC.Resource"), EventAttributes);
@@ -236,7 +242,7 @@ void FStudioTelemetryEditor::RecordEvent_DDCResource(const FString& Context, TAr
 void FStudioTelemetryEditor::RecordEvent_DDCSummary(const FString& Context, TArray<FAnalyticsEventAttribute> Attributes)
 {
 #if ENABLE_COOK_STATS
-	const int SchemaVersion = 2;
+	const int SchemaVersion = 3;
 
 	Attributes.Emplace(TEXT("SchemaVersion"), SchemaVersion);
 	Attributes.Emplace(TEXT("LoadingName"), Context);
