Added Total_Count, Total_TimeSec and Total_SizeMB to DDC Resource telemetry

diff --git a/Engine/Plugins/StudioTelemetry/Source/Private/StudioTelemetryEditor.cpp b/Engine/Plugins/StudioTelemetry/Source/Private/StudioTelemetryEditor.cpp
--- a/Engine/Plugins/StudioTelemetry/Source/Private/StudioTelemetryEditor.cpp
+++ b/Engine/Plugins/StudioTelemetry/Source/Private/StudioTelemetryEditor.cpp
@@ -214,6 +214,8 @@ void FStudioTelemetryEditor::RecordEvent_DDCResource(const FString& Context, TAr
 	for (const FDerivedDataCacheResourceStat& Stat : ResourceStats)
 	{
 		const int64 TotalCount = Stat.BuildCount + Stat.LoadCount;
+		const double TotalTimeSec = Stat.BuildTimeSec + Stat.LoadTimeSec;
+		const int64 TotalSizeMB = Stat.BuildSizeMB + Stat.LoadSizeMB;
 
 		if (Stat.AssetType.IsEmpty() || TotalCount==0)
 		{
@@ -230,6 +232,9 @@ void FStudioTelemetryEditor::RecordEvent_DDCResource(const FString& Context, TAr
 		EventAttributes.Emplace(TEXT("Build_Count"), Stat.BuildCount);
 		EventAttributes.Emplace(TEXT("Build_TimeSec"), Stat.BuildTimeSec);
 		EventAttributes.Emplace(TEXT("Build_SizeMB"), Stat.BuildSizeMB);
+		EventAttributes.Emplace(TEXT("Total_Count"), TotalCount);
+		EventAttributes.Emplace(TEXT("Total_TimeSec"), TotalTimeSec);
+		EventAttributes.Emplace(TEXT("Total_SizeMB"), TotalSizeMB);
 		EventAttributes.Emplace(TEXT("Efficiency"), double(Stat.LoadCount)/double(TotalCount) );
 		EventAttributes.Emplace(TEXT("Thread_TimeSec"), Stat.GameThreadTimeSec);
 
