Fix for AssetType attribrute name

diff --git a/Engine/Plugins/StudioTelemetry/Source/Private/StudioTelemetryEditor.cpp b/Engine/Plugins/StudioTelemetry/Source/Private/StudioTelemetryEditor.cpp
--- a/Engine/Plugins/StudioTelemetry/Source/Private/StudioTelemetryEditor.cpp
+++ b/Engine/Plugins/StudioTelemetry/Source/Private/StudioTelemetryEditor.cpp
@@ -225,7 +225,7 @@ void FStudioTelemetryEditor::RecordEvent_DDCResource(const FString& Context, TAr
 	
 		TArray<FAnalyticsEventAttribute> EventAttributes = Attributes;
 
-		EventAttributes.Emplace(TEXT("Asset_Type"), Stat.AssetType);
+		EventAttributes.Emplace(TEXT("AssetType"), Stat.AssetType);
 		EventAttributes.Emplace(TEXT("Load_Count"), Stat.LoadCount);
 		EventAttributes.Emplace(TEXT("Load_TimeSec"), Stat.LoadTimeSec);
 		EventAttributes.Emplace(TEXT("Load_SizeMB"), Stat.LoadSizeMB);
