Pushed Set/GetDefaultAttributes into IAnalyticsProvider
Fixed up FileLogging, AnalyticsSwrve and AnalyticsMulticast to implement the IAnalyticsProvider changes to Set/GetDefaultAttributes
Replaced dependency on IAnalyticsProviderET with IAnalyticsProvider in Studio Telemetry
Added Dynamic Analytics Provider Module Loading support. Replaced hard coded provider creation with dynamic IAnalyticsModule interface through modules
Added AnalyticsLog module as an example for 3rd parties to add custom, data driven IAnalyticsProviders to Studio Telemetry
Updated the ini files to use ProviderModule as the factory type for Porivder creation

diff --git a/Engine/Plugins/Experimental/StudioTelemetry/Source/AnalyticsLog/AnalyticsLog.Build.cs b/Engine/Plugins/Experimental/StudioTelemetry/Source/AnalyticsLog/AnalyticsLog.Build.cs
new file mode 100644
--- /dev/null
+++ b/Engine/Plugins/Experimental/StudioTelemetry/Source/AnalyticsLog/AnalyticsLog.Build.cs
@@ -0,0 +1,24 @@
+// Copyright Epic Games, Inc. All Rights Reserved.
+
+namespace UnrealBuildTool.Rules
+{
+	public class AnalyticsLog : ModuleRules
+	{
+		public AnalyticsLog(ReadOnlyTargetRules Target) : base(Target)
+		{
+			PublicDependencyModuleNames.AddRange(
+				new string[] {
+						"Core"
+				}
+			);
+
+			PrivateDependencyModuleNames.AddRange(
+				new string[] {
+				"CoreUObject",
+				"Analytics",
+				"AnalyticsET"
+				}
+			);
+		}
+	}
+}
diff --git a/Engine/Plugins/Experimental/StudioTelemetry/Source/AnalyticsLog/Private/AnalyticsLog.cpp b/Engine/Plugins/Experimental/StudioTelemetry/Source/AnalyticsLog/Private/AnalyticsLog.cpp
new file mode 100644
--- /dev/null
+++ b/Engine/Plugins/Experimental/StudioTelemetry/Source/AnalyticsLog/Private/AnalyticsLog.cpp
@@ -0,0 +1,25 @@
+// Copyright Epic Games, Inc. All Rights Reserved.
+
+#include "AnalyticsLog.h"
+#include "AnalyticsProviderLog.h"
+
+IMPLEMENT_MODULE( FAnalyticsLog, AnalyticsLog );
+
+void FAnalyticsLog::StartupModule()
+{
+}
+
+void FAnalyticsLog::ShutdownModule()
+{
+}
+
+TSharedPtr<IAnalyticsProvider> FAnalyticsLog::CreateAnalyticsProvider(const FAnalyticsProviderConfigurationDelegate& GetConfigValue) const
+{
+	if (GetConfigValue.IsBound())
+	{
+		return MakeShared<FAnalyticsProviderLog>(GetConfigValue);
+	}
+
+	return nullptr;
+}
+
diff --git a/Engine/Plugins/Experimental/StudioTelemetry/Source/AnalyticsLog/Private/AnalyticsProviderLog.cpp b/Engine/Plugins/Experimental/StudioTelemetry/Source/AnalyticsLog/Private/AnalyticsProviderLog.cpp
new file mode 100644
--- /dev/null
+++ b/Engine/Plugins/Experimental/StudioTelemetry/Source/AnalyticsLog/Private/AnalyticsProviderLog.cpp
@@ -0,0 +1,122 @@
+// Copyright Epic Games, Inc. All Rights Reserved.
+
+#include "AnalyticsProviderLog.h"
+#include "Analytics.h"
+#include "HAL/FileManager.h"
+#include "Misc/Paths.h"
+
+FAnalyticsProviderLog::FAnalyticsProviderLog(const FAnalyticsProviderConfigurationDelegate& GetConfigValue)
+{
+	FString FileName = GetConfigValue.Execute(TEXT("FileName"), true);
+
+	if (FileName.IsEmpty())
+	{
+		// Use default filename
+		FileName = TEXT("Telemetry.log");
+	}
+
+	FString PathName = GetConfigValue.Execute(TEXT("PathName"), true);
+
+	if (PathName.IsEmpty())
+	{
+		// Use default output path
+		PathName = FPaths::EngineSavedDir();
+	}
+
+	// Create the full output path
+	FString FilePath = PathName / FileName;
+	FileWriter = TUniquePtr<FArchive>(IFileManager::Get().CreateFileWriter(*FilePath, FILEWRITE_EvenIfReadOnly));
+}
+
+FAnalyticsProviderLog::~FAnalyticsProviderLog()
+{
+}
+
+bool FAnalyticsProviderLog::SetSessionID(const FString& InSessionID)
+{
+	SessionID = InSessionID;
+	return true;
+}
+
+FString FAnalyticsProviderLog::GetSessionID() const
+{
+	return SessionID;
+}
+
+void FAnalyticsProviderLog::SetUserID(const FString& InUserID)
+{
+	UserID = InUserID;
+
+}
+
+FString FAnalyticsProviderLog::GetUserID() const
+{
+	return UserID;
+}
+
+void FAnalyticsProviderLog::FlushEvents()
+{
+}
+
+void FAnalyticsProviderLog::SetDefaultEventAttributes(TArray<FAnalyticsEventAttribute>&& Attributes)
+{
+	DefaultEventAttributes = Attributes;
+}
+
+TArray<FAnalyticsEventAttribute> FAnalyticsProviderLog::GetDefaultEventAttributesSafe() const
+{
+	return DefaultEventAttributes;
+}
+
+int32 FAnalyticsProviderLog::GetDefaultEventAttributeCount() const
+{
+	return DefaultEventAttributes.Num();
+}
+
+FAnalyticsEventAttribute FAnalyticsProviderLog::GetDefaultEventAttribute(int AttributeIndex) const
+{
+	return DefaultEventAttributes[AttributeIndex];
+}
+
+bool FAnalyticsProviderLog::StartSession(const TArray<FAnalyticsEventAttribute>& Attributes)
+{
+	RecordEvent(TEXT("StartSession"), Attributes);
+
+	return true;
+}
+
+void FAnalyticsProviderLog::EndSession()
+{
+	RecordEvent(TEXT("EndSession"));
+
+	if (FileWriter)
+	{
+		FileWriter->Close();
+	}
+}
+
+void FAnalyticsProviderLog::RecordEvent(const FString& EventName, const TArray<FAnalyticsEventAttribute>& Attributes)
+{	
+	TStringBuilder<1024> StringBuilder;
+
+	StringBuilder.Appendf(TEXT("{{EventName=%s}"), *EventName);
+
+	// Log the default attributes
+	for (const FAnalyticsEventAttribute& Attribute : DefaultEventAttributes)
+	{
+		StringBuilder.Appendf(TEXT(",{%s=%s}"), *Attribute.GetName(), *Attribute.GetValue());
+	}
+
+	// Log the event attributes
+	for ( const FAnalyticsEventAttribute& Attribute : Attributes )
+	{
+		StringBuilder.Appendf(TEXT("{%s=%s}"), *Attribute.GetName(), *Attribute.GetValue());
+	}
+
+	StringBuilder.Appendf(TEXT("}\n"));
+
+	if (FileWriter)
+	{
+		FileWriter->Serialize(TCHAR_TO_ANSI(StringBuilder.ToString()), StringBuilder.Len());
+	}
+}
\ No newline at end of file
diff --git a/Engine/Plugins/Experimental/StudioTelemetry/Source/AnalyticsLog/Private/AnalyticsProviderLog.h b/Engine/Plugins/Experimental/StudioTelemetry/Source/AnalyticsLog/Private/AnalyticsProviderLog.h
new file mode 100644
--- /dev/null
+++ b/Engine/Plugins/Experimental/StudioTelemetry/Source/AnalyticsLog/Private/AnalyticsProviderLog.h
@@ -0,0 +1,45 @@
+// Copyright Epic Games, Inc. All Rights Reserved.
+
+#pragma once
+
+#if WITH_EDITOR
+
+#include "AnalyticsProviderConfigurationDelegate.h"
+#include "Interfaces/IAnalyticsProvider.h"
+
+/**
+ * Implementation of the IAnalyticsProviderET interface that exports to OpenTelemetry
+ */
+class FAnalyticsProviderLog : public IAnalyticsProvider
+{
+public:
+
+	FAnalyticsProviderLog(const FAnalyticsProviderConfigurationDelegate& GetConfigValue);
+	~FAnalyticsProviderLog();
+
+	virtual bool StartSession(const TArray<FAnalyticsEventAttribute>& Attributes = {}) override;
+	virtual void EndSession() override;
+
+	virtual void FlushEvents() override;
+
+	virtual void SetDefaultEventAttributes(TArray<FAnalyticsEventAttribute>&& Attributes)  override;
+	virtual TArray<FAnalyticsEventAttribute> GetDefaultEventAttributesSafe() const override;
+	virtual int32 GetDefaultEventAttributeCount() const  override;
+	virtual FAnalyticsEventAttribute GetDefaultEventAttribute(int AttributeIndex) const  override;
+	virtual bool SetSessionID(const FString& InSessionID) override;
+	virtual void SetUserID(const FString& InUserID) override;
+
+	virtual FString GetSessionID() const override;
+	virtual FString GetUserID() const override;
+	
+	virtual void RecordEvent(const FString& EventName, const TArray<FAnalyticsEventAttribute>& Attributes = {}) override;
+
+private:
+
+	FString									UserID;
+	FString									SessionID;
+	TArray<FAnalyticsEventAttribute>		DefaultEventAttributes;	
+	TUniquePtr<FArchive>					FileWriter;
+};
+
+#endif
\ No newline at end of file
diff --git a/Engine/Plugins/Experimental/StudioTelemetry/Source/AnalyticsLog/Public/AnalyticsLog.h b/Engine/Plugins/Experimental/StudioTelemetry/Source/AnalyticsLog/Public/AnalyticsLog.h
new file mode 100644
--- /dev/null
+++ b/Engine/Plugins/Experimental/StudioTelemetry/Source/AnalyticsLog/Public/AnalyticsLog.h
@@ -0,0 +1,63 @@
+// Copyright Epic Games, Inc. All Rights Reserved.
+
+#pragma once
+
+#if WITH_EDITOR
+
+#include "AnalyticsProviderConfigurationDelegate.h"
+#include "AnalyticsET.h"
+#include "Containers/Array.h"
+#include "Containers/UnrealString.h"
+#include "CoreMinimal.h"
+#include "HAL/Platform.h"
+#include "HAL/PlatformCrt.h"
+#include "Interfaces/IAnalyticsProviderModule.h"
+#include "Modules/ModuleManager.h"
+#include "Templates/SharedPointer.h"
+#include "Templates/UnrealTemplate.h"
+
+class IAnalyticsProvider;
+class IAnalyticsProviderET;
+
+/**
+ *  Public implementation of EpicGames.MCP.AnalyticsProvider
+ */
+class FAnalyticsLog : public IAnalyticsProviderModule
+{
+	//--------------------------------------------------------------------------
+	// Module functionality
+	//--------------------------------------------------------------------------
+public:
+	/**
+	 * Singleton-like access to this module's interface.  This is just for convenience!
+	 * Beware of calling this during the shutdown phase, though.  Your module might have been unloaded already.
+	 *
+	 * @return Returns singleton instance, loading the module on demand if needed
+	 */
+	static inline FAnalyticsLog& Get()
+	{
+		return FModuleManager::LoadModuleChecked< FAnalyticsLog >( "AnalyticsLog" );
+	}
+
+	//--------------------------------------------------------------------------
+	// Configuration functionality
+	//--------------------------------------------------------------------------
+public:
+	
+	//--------------------------------------------------------------------------
+	// provider factory functions
+	//--------------------------------------------------------------------------
+public:
+	/**
+	 * IAnalyticsProviderModule interface.
+	 * Creates the analytics provider given a configuration delegate.
+	 * The keys required exactly match the field names in the Config object. 
+	 */
+	ANALYTICSLOG_API virtual TSharedPtr<IAnalyticsProvider> CreateAnalyticsProvider(const FAnalyticsProviderConfigurationDelegate& GetConfigValue) const override;
+
+private:
+	ANALYTICSLOG_API virtual void StartupModule() override;
+	ANALYTICSLOG_API virtual void ShutdownModule() override;
+};
+
+#endif
\ No newline at end of file
diff --git a/Engine/Plugins/Experimental/StudioTelemetry/Source/Private/AnalyticsProviderMulticast.cpp b/Engine/Plugins/Experimental/StudioTelemetry/Source/Private/AnalyticsProviderMulticast.cpp
deleted file mode 100644
diff --git a/Engine/Plugins/Experimental/StudioTelemetry/Source/Private/AnalyticsProviderMulticast.h b/Engine/Plugins/Experimental/StudioTelemetry/Source/Private/AnalyticsProviderMulticast.h
deleted file mode 100644
diff --git a/Engine/Plugins/Experimental/StudioTelemetry/Source/Private/StudioTelemetry.cpp b/Engine/Plugins/Experimental/StudioTelemetry/Source/Private/StudioTelemetry.cpp
deleted file mode 100644
diff --git a/Engine/Plugins/Experimental/StudioTelemetry/Source/Private/StudioTelemetryEditor.cpp b/Engine/Plugins/Experimental/StudioTelemetry/Source/Private/StudioTelemetryEditor.cpp
deleted file mode 100644
diff --git a/Engine/Plugins/Experimental/StudioTelemetry/Source/Private/StudioTelemetryEditor.h b/Engine/Plugins/Experimental/StudioTelemetry/Source/Private/StudioTelemetryEditor.h
deleted file mode 100644
diff --git a/Engine/Plugins/Experimental/StudioTelemetry/Source/Private/StudioTelemetryLog.h b/Engine/Plugins/Experimental/StudioTelemetry/Source/Private/StudioTelemetryLog.h
deleted file mode 100644
diff --git a/Engine/Plugins/Experimental/StudioTelemetry/Source/Public/StudioTelemetry.h b/Engine/Plugins/Experimental/StudioTelemetry/Source/Public/StudioTelemetry.h
deleted file mode 100644
diff --git a/Engine/Plugins/Experimental/StudioTelemetry/Source/StudioTelemetry.build.cs b/Engine/Plugins/Experimental/StudioTelemetry/Source/StudioTelemetry.build.cs
deleted file mode 100644
diff --git a/Engine/Plugins/Experimental/StudioTelemetry/Source/StudioTelemetry/Private/AnalyticsProviderMulticast.cpp b/Engine/Plugins/Experimental/StudioTelemetry/Source/StudioTelemetry/Private/AnalyticsProviderMulticast.cpp
new file mode 100644
--- /dev/null
+++ b/Engine/Plugins/Experimental/StudioTelemetry/Source/StudioTelemetry/Private/AnalyticsProviderMulticast.cpp
@@ -0,0 +1,236 @@
+// Copyright Epic Games, Inc. All Rights Reserved.
+
+#include "AnalyticsProviderMulticast.h"
+#include "StudioTelemetryLog.h"
+#include "Analytics.h"
+#include "Interfaces/IAnalyticsProviderModule.h"
+#include "AnalyticsProviderConfigurationDelegate.h"
+#include "Misc/ConfigCacheIni.h"
+#include "HttpModule.h"
+
+const FString TelemetrySection(TEXT("StudioTelemetry"));
+static FString ProviderSection;
+
+FString GetAnalyticsProviderConfiguration(const FString& Name, bool)
+{
+	FString Result;
+	GConfig->GetString(*ProviderSection, *Name, Result, GEngineIni);
+	return Result;
+}
+
+TSharedPtr<FAnalyticsProviderMulticast> FAnalyticsProviderMulticast::CreateAnalyticsProvider()
+{
+	return MakeShared<FAnalyticsProviderMulticast>();
+}
+
+TWeakPtr<IAnalyticsProvider> FAnalyticsProviderMulticast::GetAnalyticsProvider(const FString& Name)
+{
+	TSharedPtr<IAnalyticsProvider>* ProviderPtr = Providers.Find(Name);
+	return ProviderPtr != nullptr ? *ProviderPtr : TSharedPtr<IAnalyticsProvider>();
+}
+
+FAnalyticsProviderMulticast::FAnalyticsProviderMulticast()
+{
+	TArray<FString> SectionNames;
+	
+	if (GConfig->GetSectionNames(GEngineIni, SectionNames))
+	{
+		for (const FString& SectionName : SectionNames)
+		{
+			if (SectionName.Find(TelemetrySection) != INDEX_NONE)
+			{
+				ProviderSection = SectionName;
+
+				FString UsageType;
+
+				// Validate the usage type is for this build type
+				if (GConfig->GetString(*ProviderSection, TEXT("UsageType"), UsageType, GEngineIni))
+				{
+#if WITH_EDITOR
+					// Must specify a Editor usage type for this type build
+					if (UsageType.Find(TEXT("Editor")) == INDEX_NONE)
+					{
+						continue;
+					}
+#elif WITH_SERVER_CODE
+					// Must specify a Server usage type for this type build
+					if (UsageType.Find(TEXT("Server")) == INDEX_NONE)
+					{
+						continue;
+					}
+#else
+					// Must specify a Client or Program usage type for this type build
+					if (UsageType.Find(TEXT("Program")) == INDEX_NONE && UsageType.Find(TEXT("Client")) == INDEX_NONE)
+					{
+						continue;
+					}
+#endif
+				}
+				else
+				{
+					// Must always specify a usage type
+					UE_LOG(LogStudioTelemetry, Error, TEXT("There must be a valid UsageType specified for analytics provider %s"), *ProviderSection);
+					continue;
+				}
+
+				FString ProviderModuleName;
+
+				if (GConfig->GetString(*ProviderSection, TEXT("ProviderModule"), ProviderModuleName, GEngineIni))
+				{
+					TSharedPtr<IAnalyticsProvider> AnalyticsProvider;
+
+					FString Name = GetAnalyticsProviderConfiguration("Name", true);
+
+					if ( Name.IsEmpty() )
+					{ 
+						UE_LOG(LogStudioTelemetry, Error, TEXT("There must be a valid Name specified for analytics provider %s."), *ProviderSection);
+						continue;
+					}
+					else if (Providers.Find(Name) )
+					{
+						UE_LOG(LogStudioTelemetry, Warning, TEXT("An analytics provider with name %s already exists."), *Name);
+						continue;
+					}
+
+					// Try to create the analytics provider
+					AnalyticsProvider = FAnalytics::Get().CreateAnalyticsProvider(FName(ProviderModuleName), FAnalyticsProviderConfigurationDelegate::CreateStatic(&GetAnalyticsProviderConfiguration));
+	
+					if (AnalyticsProvider.IsValid())
+					{
+						UE_LOG(LogStudioTelemetry, Display, TEXT("Created an analytics provider %s from module %s configuration %s [%s]"), *Name, *ProviderModuleName, *GEngineIni, *ProviderSection);
+						Providers.Add(Name, AnalyticsProvider);
+					}
+					else
+					{
+						UE_LOG(LogStudioTelemetry, Warning, TEXT("Unable to create an analytics provider %s from module %s configuration %s [%s]"), *Name, *ProviderModuleName, *GEngineIni, *ProviderSection);
+					}
+				}
+				else
+				{
+					UE_LOG(LogStudioTelemetry, Error, TEXT("There must be a valid ProviderModule specified for analytics provider %s"), *ProviderSection);
+				}
+			}
+		}
+	}
+}
+
+bool FAnalyticsProviderMulticast::SetSessionID(const FString& InSessionID)
+{
+	SessionID = InSessionID;
+
+	bool bResult = true;
+
+	for (TProviders::TConstIterator it(Providers); it; ++it)
+	{
+		bResult &= (*it).Value->SetSessionID(InSessionID);
+	}
+
+	return bResult;
+}
+
+FString FAnalyticsProviderMulticast::GetSessionID() const
+{
+	return SessionID;
+}
+
+void FAnalyticsProviderMulticast::SetUserID(const FString& InUserID)
+{
+	UserID = InUserID;
+
+	for (TProviders::TConstIterator it(Providers); it; ++it)
+	{
+		(*it).Value->SetUserID(InUserID);
+	}
+}
+
+FString FAnalyticsProviderMulticast::GetUserID() const
+{
+	return UserID;
+}
+
+void FAnalyticsProviderMulticast::SetDefaultEventAttributes(TArray<FAnalyticsEventAttribute>&& Attributes)
+{
+	DefaultEventAttributes = Attributes;
+
+	for (TProviders::TConstIterator it(Providers); it; ++it)
+	{
+		(*it).Value->SetDefaultEventAttributes(CopyTemp(DefaultEventAttributes));
+	}
+}
+
+TArray<FAnalyticsEventAttribute> FAnalyticsProviderMulticast::GetDefaultEventAttributesSafe() const
+{
+	return DefaultEventAttributes;
+}
+
+int32 FAnalyticsProviderMulticast::GetDefaultEventAttributeCount() const
+{
+	return DefaultEventAttributes.Num();
+}
+
+FAnalyticsEventAttribute FAnalyticsProviderMulticast::GetDefaultEventAttribute(int AttributeIndex) const
+{
+	return DefaultEventAttributes[AttributeIndex];
+}
+
+bool FAnalyticsProviderMulticast::StartSession(const TArray<FAnalyticsEventAttribute>& Attributes)
+{
+	bool bResult = true;
+
+	for (TProviders::TConstIterator it(Providers);it;++it)
+	{
+		bResult &= (*it).Value->StartSession(Attributes);
+	}
+	return bResult;
+}
+
+void FAnalyticsProviderMulticast::EndSession()
+{
+	for (TProviders::TConstIterator it(Providers);it;++it)
+	{
+		TSharedPtr<IAnalyticsProvider> Provider = (*it).Value;
+
+		Provider->EndSession();
+		Provider.Reset();
+	}
+
+	Providers.Reset();
+}
+
+void FAnalyticsProviderMulticast::FlushEvents()
+{
+	for (TProviders::TConstIterator it(Providers); it; ++it)
+	{
+		(*it).Value->FlushEvents();
+	}
+}
+
+void FAnalyticsProviderMulticast::SetRecordEventCallback(OnRecordEvent Callback)
+{
+	RecordEventCallback = Callback;
+}
+
+void FAnalyticsProviderMulticast::RecordEvent(const FString& EventName, const TArray<FAnalyticsEventAttribute>& Attributes)
+{
+#if UE_BUILD_DEBUG || UE_BUILD_DEVELOPMENT
+	// Expose events that have duplicate aatibute names. This is is not handled by the analytics backends in any reliable way.
+	for (int32 index0 = 0; index0 < Attributes.Num(); ++index0)
+	{
+		for (int32 index1 = index0 + 1; index1 < Attributes.Num(); ++index1)
+		{
+			checkf(Attributes[index0].GetName() != Attributes[index1].GetName(), TEXT("Duplicate Attributes Found For Event %s %s==%s"), *EventName, *Attributes[index0].GetName(), *Attributes[index1].GetName());
+		}
+	}
+#endif
+
+	for (TProviders::TConstIterator it(Providers); it; ++it)
+	{
+		(*it).Value->RecordEvent(EventName, Attributes);
+	}
+
+	if (RecordEventCallback)
+	{
+		// Notify any callbacks
+		RecordEventCallback(EventName, Attributes);
+	}
+}
diff --git a/Engine/Plugins/Experimental/StudioTelemetry/Source/StudioTelemetry/Private/AnalyticsProviderMulticast.h b/Engine/Plugins/Experimental/StudioTelemetry/Source/StudioTelemetry/Private/AnalyticsProviderMulticast.h
new file mode 100644
--- /dev/null
+++ b/Engine/Plugins/Experimental/StudioTelemetry/Source/StudioTelemetry/Private/AnalyticsProviderMulticast.h
@@ -0,0 +1,51 @@
+// Copyright Epic Games, Inc. All Rights Reserved.
+
+#pragma once
+
+#include "AnalyticsProviderConfigurationDelegate.h"
+#include "Interfaces/IAnalyticsProvider.h"
+
+/**
+ * Implementation of the IAnalyticsProvider interface that forwards the API calls to an array of IAnalyticsProviderET interfaces
+ */
+class FAnalyticsProviderMulticast : public IAnalyticsProvider
+{
+public:
+
+	using TProviders = TMap<FString, TSharedPtr<IAnalyticsProvider>>;
+
+	typedef TFunction<void(const FString& EventName, const TArray<FAnalyticsEventAttribute>& Attrs)> OnRecordEvent;
+
+
+	FAnalyticsProviderMulticast();
+
+	static TSharedPtr<FAnalyticsProviderMulticast> CreateAnalyticsProvider();
+
+	virtual bool StartSession(const TArray<FAnalyticsEventAttribute>& Attributes = {}) override;
+	virtual void EndSession() override;
+	virtual void FlushEvents() override;
+	
+	virtual void SetDefaultEventAttributes(TArray<FAnalyticsEventAttribute>&& Attributes)  override;
+	virtual TArray<FAnalyticsEventAttribute> GetDefaultEventAttributesSafe() const override;
+	virtual int32 GetDefaultEventAttributeCount() const  override;
+	virtual FAnalyticsEventAttribute GetDefaultEventAttribute(int AttributeIndex) const  override;
+
+	virtual bool SetSessionID(const FString& InSessionID) override;
+	virtual void SetUserID(const FString& InUserID) override;
+
+	virtual FString GetSessionID() const override;
+	virtual FString GetUserID() const override;
+	
+	virtual void RecordEvent(const FString& EventName, const TArray<FAnalyticsEventAttribute>& Attributes) override;	
+	void SetRecordEventCallback(OnRecordEvent Callback);
+	bool HasValidProviders() const { return Providers.Num() > 0; }
+	TWeakPtr<IAnalyticsProvider> GetAnalyticsProvider(const FString& Name);
+
+private:
+
+	TProviders								Providers;
+	FString									UserID;
+	FString									SessionID;
+	TArray<FAnalyticsEventAttribute>		DefaultEventAttributes;
+	OnRecordEvent							RecordEventCallback;
+};
diff --git a/Engine/Plugins/Experimental/StudioTelemetry/Source/StudioTelemetry/Private/StudioTelemetry.cpp b/Engine/Plugins/Experimental/StudioTelemetry/Source/StudioTelemetry/Private/StudioTelemetry.cpp
new file mode 100644
--- /dev/null
+++ b/Engine/Plugins/Experimental/StudioTelemetry/Source/StudioTelemetry/Private/StudioTelemetry.cpp
@@ -0,0 +1,217 @@
+// Copyright Epic Games, Inc. All Rights Reserved.
+
+#include "StudioTelemetry.h"
+#include "StudioTelemetryLog.h"
+
+#if WITH_EDITOR
+#include "StudioTelemetryEditor.h"
+#include "Horde.h"
+#endif
+
+#include "Analytics.h"
+#include "AnalyticsProviderMulticast.h"
+#include "BuildSettings.h"
+#include "RHI.h"
+
+#include "HAL/PlatformTime.h"
+#include "HAL/PlatformProcess.h"
+#include "HAL/Thread.h"
+
+#include "Misc/CommandLine.h"
+#include "Misc/App.h"
+#include "Misc/EngineVersion.h"
+#include "Misc/EngineBuildSettings.h"
+
+#include "GenericPlatform/GenericPlatformMisc.h"
+
+DEFINE_LOG_CATEGORY(LogStudioTelemetry);
+
+IMPLEMENT_MODULE(FStudioTelemetry, StudioTelemetry)
+
+FStudioTelemetry& FStudioTelemetry::Get()
+{
+	static FStudioTelemetry StudioTelemetryInstance;
+	return StudioTelemetryInstance;
+}
+
+void FStudioTelemetry::SetRecordEventCallback(OnRecordEvent Callback )
+{
+	RecordEventCallback = Callback;
+
+	// If the provider already exists then set the callback
+	if (AnalyticsProvider.IsValid())
+	{	
+		AnalyticsProvider->SetRecordEventCallback(RecordEventCallback);
+	}
+}
+
+void FStudioTelemetry::StartupModule()
+{
+	UE_LOG(LogStudioTelemetry, Display, TEXT("Starting StudioTelemetry Module"));
+
+	// Create the provider and start the analytics session
+	FStudioTelemetry::Get().StartSession();
+
+#if WITH_EDITOR
+	// Initialize the analytics subsystems 
+	FStudioTelemetryEditor::Get().Initialize();
+#endif
+}
+
+void FStudioTelemetry::ShutdownModule()
+{
+#if WITH_EDITOR
+	// Shutdown the analytics subsystems
+	FStudioTelemetryEditor::Get().Shutdown();
+#endif
+
+	// End the session and destroy analytics provider
+	FStudioTelemetry::Get().EndSession();
+
+	UE_LOG(LogStudioTelemetry, Display, TEXT("Shutdown StudioTelemetry Module"));
+}
+
+void FStudioTelemetry::EndSession()
+{
+	if (AnalyticsFlowTracker.IsValid())
+	{
+		AnalyticsFlowTracker->EndSession();
+		AnalyticsFlowTracker.Reset();
+	}
+
+	if (AnalyticsProvider.IsValid())
+	{
+		AnalyticsProvider->EndSession();
+		AnalyticsProvider.Reset();
+	}
+
+	UE_LOG(LogStudioTelemetry, Log, TEXT("Ended StudioTelemetry Session"));
+}
+
+void FStudioTelemetry::StartSession()
+{
+	AnalyticsProvider = FAnalyticsProviderMulticast::CreateAnalyticsProvider();
+
+	if (AnalyticsProvider.IsValid())
+	{
+		TArray<FAnalyticsEventAttribute> DefaultEventAttributes;
+
+		const FString UserID = FPlatformProcess::UserName(false);
+		const FString ProjectName = FApp::GetProjectName();
+		FString ComputerName = FPlatformProcess::ComputerName();
+		
+		FString ProjectIDString;
+		GConfig->GetString(TEXT("/Script/EngineSettings.GeneralProjectSettings"), TEXT("ProjectID"), ProjectIDString, GGameIni);
+
+		FGuid ProjectID;
+
+		TArray<FString> Elements;
+		if (ProjectIDString.ParseIntoArray(Elements, TEXT("=")) == 5) 
+		{
+			ProjectID = FGuid(FCString::Atoi(*(Elements[1])), FCString::Atoi(*(Elements[2])), FCString::Atoi(*(Elements[3])), FCString::Atoi(*(Elements[4])));
+		}
+
+		FGuid SessionGUID;
+		FPlatformMisc::CreateGuid(SessionGUID);
+		
+		FString SessionLabel;
+		FParse::Value(FCommandLine::Get(), TEXT("SessionLabel="), SessionLabel);
+
+		// Set the default event attributes
+		DefaultEventAttributes.Emplace(TEXT("User_ID"), UserID);
+		DefaultEventAttributes.Emplace(TEXT("Application_Commandline"), FCommandLine::Get());
+
+		DefaultEventAttributes.Emplace(TEXT("ProjectName"), ProjectName);
+		DefaultEventAttributes.Emplace(TEXT("ProjectID"), ProjectID);
+
+		DefaultEventAttributes.Emplace(TEXT("Session_ID"), SessionGUID.ToString(EGuidFormats::DigitsWithHyphensInBraces));
+		DefaultEventAttributes.Emplace(TEXT("Session_Label"), SessionLabel);
+		DefaultEventAttributes.Emplace(TEXT("Session_StartUTC"), FDateTime::UtcNow().ToUnixTimestampDecimal());
+
+		DefaultEventAttributes.Emplace(TEXT("Build_Configuration"), LexToString(FApp::GetBuildConfiguration()));
+		DefaultEventAttributes.Emplace(TEXT("Build_IsInternalBuild"), FEngineBuildSettings::IsInternalBuild());
+		DefaultEventAttributes.Emplace(TEXT("Build_IsPerforceBuild"), FEngineBuildSettings::IsPerforceBuild());
+		DefaultEventAttributes.Emplace(TEXT("Build_IsPromotedBuild"), FApp::GetEngineIsPromotedBuild() == 0 ? false : true);
+		DefaultEventAttributes.Emplace(TEXT("Build_BranchName"), FApp::GetBranchName().ToLower());
+		DefaultEventAttributes.Emplace(TEXT("Build_Changelist"), BuildSettings::GetCurrentChangelist());
+
+		DefaultEventAttributes.Emplace(TEXT("Hardware_ComputerName"), ComputerName.ToLower());
+		DefaultEventAttributes.Emplace(TEXT("Hardware_Platform"), FString(FPlatformProperties::IniPlatformName()));
+		DefaultEventAttributes.Emplace(TEXT("Hardware_GPU"), GRHIAdapterName);
+		DefaultEventAttributes.Emplace(TEXT("Hardware_CPU"), FPlatformMisc::GetCPUBrand());
+		DefaultEventAttributes.Emplace(TEXT("Hardware_CPU_Cores_Physical"), FPlatformMisc::NumberOfCores());
+		DefaultEventAttributes.Emplace(TEXT("Hardware_CPU_Cores_Logical"), FPlatformMisc::NumberOfCoresIncludingHyperthreads());
+		DefaultEventAttributes.Emplace(TEXT("Hardware_RAM"), static_cast<uint64>(FPlatformMemory::GetStats().TotalPhysical));
+
+		DefaultEventAttributes.Emplace(TEXT("Config_IsEditor"), GIsEditor);
+		DefaultEventAttributes.Emplace(TEXT("Config_IsUnattended"), FApp::IsUnattended());
+		DefaultEventAttributes.Emplace(TEXT("Config_IsBuildMachine"), GIsBuildMachine);
+		DefaultEventAttributes.Emplace(TEXT("Config_IsRunningCommandlet"), IsRunningCommandlet());
+
+#if WITH_EDITOR
+		if (!FHorde::GetJobId().IsEmpty())
+		{
+			// Only send Horde data if applicable
+			DefaultEventAttributes.Emplace(TEXT("Horde_TemplateID"), FHorde::GetTemplateId());
+			DefaultEventAttributes.Emplace(TEXT("Horde_TemplateName"), FHorde::GetTemplateName());
+			DefaultEventAttributes.Emplace(TEXT("Horde_JobURL"), FHorde::GetJobURL());
+			DefaultEventAttributes.Emplace(TEXT("Horde_JobID"), FHorde::GetJobId());
+			DefaultEventAttributes.Emplace(TEXT("Horde_StepName"), FHorde::GetStepName());
+			DefaultEventAttributes.Emplace(TEXT("Horde_StepID"), FHorde::GetStepId());
+			DefaultEventAttributes.Emplace(TEXT("Horde_StepURL"), FHorde::GetStepURL());
+			DefaultEventAttributes.Emplace(TEXT("Horde_BatchID"), FHorde::GetBatchId());
+		}
+#endif
+		
+		// Set up the analytics provider
+		AnalyticsProvider->SetUserID(UserID);
+		AnalyticsProvider->SetSessionID(SessionGUID.ToString(EGuidFormats::DigitsWithHyphensInBraces));
+		AnalyticsProvider->SetDefaultEventAttributes(MoveTemp(DefaultEventAttributes));
+		AnalyticsProvider->SetRecordEventCallback(RecordEventCallback);
+		
+		// Start the analytics session
+		AnalyticsProvider->StartSession();
+
+		// Make the flow tracker and start the session
+		AnalyticsFlowTracker = MakeShared<FAnalyticsFlowTracker>();
+		AnalyticsFlowTracker->SetProvider(AnalyticsProvider);
+		AnalyticsFlowTracker->StartSession();
+
+		UE_LOG(LogStudioTelemetry, Log, TEXT("Started StudioTelemetry Session"));
+	}
+}
+
+void FStudioTelemetry::RecordEvent(const FString& EventName, const TArray<FAnalyticsEventAttribute>& Attributes)
+{
+	if (AnalyticsProvider.IsValid())
+	{
+		FScopeLock ScopeLock(&CriticalSection);
+		AnalyticsProvider->RecordEvent(EventName, Attributes);
+	}
+}
+
+void FStudioTelemetry::RecordEvent(const FString& ProviderName, const FString& EventName, const TArray<FAnalyticsEventAttribute>& Attributes)
+{
+	FScopeLock ScopeLock(&CriticalSection);
+	TSharedPtr<IAnalyticsProvider> NamedProvider = GetProvider(ProviderName).Pin();
+
+	if (NamedProvider.IsValid())
+	{
+		NamedProvider->RecordEvent(EventName, Attributes);
+	}
+}
+
+TWeakPtr<IAnalyticsProvider> FStudioTelemetry::GetProvider()
+{
+	return AnalyticsProvider;
+}
+
+TWeakPtr<IAnalyticsProvider> FStudioTelemetry::GetProvider(const FString& Name)
+{
+	return AnalyticsProvider.IsValid()? AnalyticsProvider->GetAnalyticsProvider(Name) : TWeakPtr<IAnalyticsProvider>();
+}
+
+TWeakPtr<FAnalyticsFlowTracker> FStudioTelemetry::GetFlowTracker()
+{
+	return AnalyticsFlowTracker;
+}
\ No newline at end of file
diff --git a/Engine/Plugins/Experimental/StudioTelemetry/Source/StudioTelemetry/Private/StudioTelemetryEditor.cpp b/Engine/Plugins/Experimental/StudioTelemetry/Source/StudioTelemetry/Private/StudioTelemetryEditor.cpp
new file mode 100644
--- /dev/null
+++ b/Engine/Plugins/Experimental/StudioTelemetry/Source/StudioTelemetry/Private/StudioTelemetryEditor.cpp
@@ -0,0 +1,876 @@
+// Copyright Epic Games, Inc. All Rights Reserved.
+
+#include "StudioTelemetryEditor.h"
+
+#if WITH_EDITOR
+
+#include "StudioTelemetry.h"
+#include "AnalyticsFlowTracker.h"
+#include "CollectionManagerModule.h"
+#include "ContentBrowserModule.h"
+#include "TelemetryRouter.h"
+#include "Policies/CondensedJsonPrintPolicy.h"
+#include "AssetRegistry/AssetRegistryTelemetry.h"
+#include "Editor.h"
+#include "Editor/UnrealEdEngine.h"
+#include "UnrealEdGlobals.h"
+#include "CookOnTheSide/CookOnTheFlyServer.h"
+#include "Subsystems/AssetEditorSubsystem.h"
+#include "Misc/FeedbackContext.h"
+#include "ProfilingDebugging/CookStats.h"
+#include "DerivedDataCacheInterface.h"
+#include "DerivedDataCacheUsageStats.h"
+#include "Virtualization/VirtualizationSystem.h"
+#include "FileHelpers.h"
+#include "Experimental/ZenServerInterface.h"
+#include "ContentBrowserTelemetry.h"
+#include "ShaderStats.h"
+#include "UObject/ICookInfo.h"
+#include "IO/IoStoreOnDemand.h"
+
+UE_DISABLE_OPTIMIZATION_SHIP
+
+namespace Private
+{
+	const FName ContentBrowserModuleName = TEXT("ContentBrowser");
+	
+	// Json writer subclass to allow us to avoid using a SharedPtr to write basic Json.
+	typedef TCondensedJsonPrintPolicy<TCHAR> FPrintPolicy;
+	class FAnalyticsJsonWriter : public TJsonStringWriter<FPrintPolicy>
+	{
+	public:
+		explicit FAnalyticsJsonWriter(FString* Out) : TJsonStringWriter<FPrintPolicy>(Out, 0)
+		{
+		}
+	};
+}
+
+const TCHAR* LexToString(ECollectionTelemetryAssetAddedWorkflow Enum)
+{
+	switch(Enum)
+	{
+	case ECollectionTelemetryAssetAddedWorkflow::ContextMenu: return TEXT("ContextMenu");
+	case ECollectionTelemetryAssetAddedWorkflow::DragAndDrop: return TEXT("DragAndDrop");
+	default: return TEXT("");
+	}
+}
+	
+const TCHAR* LexToString(ECollectionTelemetryAssetRemovedWorkflow Enum)
+{
+	switch(Enum)
+	{
+	case ECollectionTelemetryAssetRemovedWorkflow::ContextMenu: return TEXT("ContextMenu");
+	default: return TEXT("");
+	}
+}
+
+template<typename T>
+FString AnalyticsOptionalToStringOrNull(const TOptional<T>& Opt)
+{
+	return Opt.IsSet() ? AnalyticsConversionToString(Opt.GetValue()) : FString(TEXT("null"));
+}
+
+FStudioTelemetryEditor& FStudioTelemetryEditor::Get()
+{
+	static FStudioTelemetryEditor StudioTelemetryEditorInstance = FStudioTelemetryEditor();
+	return StudioTelemetryEditorInstance;
+}
+
+void FStudioTelemetryEditor::RecordEvent_Cooking(TArray<FAnalyticsEventAttribute> Attributes)
+{
+#if ENABLE_COOK_STATS
+
+	const int SchemaVersion = 3;
+
+	Attributes.Emplace(TEXT("SchemaVersion"), SchemaVersion);
+
+	TMap<FString,FAnalyticsEventAttribute> CookAttributes;
+	
+	// Sends each cook stat to the studio analytics system.
+	auto GatherAnalyticsAttributes = [&CookAttributes](const FString& StatName, const TArray<FCookStatsManager::StringKeyValue>& StatAttributes)
+	{
+		for (const auto& Attr : StatAttributes)
+		{
+			const FString FormattedAttrName = (StatName + "_" + Attr.Key).Replace(TEXT("."), TEXT("_"));
+
+			if (CookAttributes.Find(FormattedAttrName)==nullptr)
+			{
+				CookAttributes.Emplace(FormattedAttrName, Attr.Value.IsNumeric() ? FAnalyticsEventAttribute(FormattedAttrName, FCString::Atof(*Attr.Value)) : FAnalyticsEventAttribute(FormattedAttrName, Attr.Value));
+			}
+		}
+	};
+
+	// Now actually grab the stats 
+	FCookStatsManager::LogCookStats(GatherAnalyticsAttributes);
+
+	// Add the values to the attributes
+	for (TMap<FString, FAnalyticsEventAttribute>::TConstIterator it(CookAttributes); it; ++it)
+	{
+		Attributes.Emplace((*it).Value);
+	}
+
+	// Gather the DDC summary stats
+	FDerivedDataCacheSummaryStats SummaryStats;
+
+	GatherDerivedDataCacheSummaryStats(SummaryStats);
+
+	// Append to the attributes
+	for (const FDerivedDataCacheSummaryStat& Stat : SummaryStats.Stats)
+	{
+		FString AttributeName = TEXT("DDC_Summary") + Stat.Key.Replace(TEXT("."), TEXT("_"));
+
+		if (Stat.Value.IsNumeric())
+		{
+			Attributes.Emplace(AttributeName, FCString::Atof(*Stat.Value));
+		}
+		else
+		{
+			Attributes.Emplace(AttributeName, Stat.Value);
+		}
+	}
+
+#if UE_WITH_ZEN
+	// Gather Zen analytics
+	if (UE::Zen::IsDefaultServicePresent())
+	{
+		UE::Zen::GetDefaultServiceInstance().GatherAnalytics(Attributes);
+	}
+#endif
+
+	if (UE::Virtualization::IVirtualizationSystem::Get().IsEnabled())
+	{
+		// Gather Virtualization analytics
+		UE::Virtualization::IVirtualizationSystem::Get().GatherAnalytics(Attributes);
+	}
+
+	FShaderStatsFunctions::GatherShaderAnalytics(Attributes);
+	
+	FStudioTelemetry::Get().RecordEvent(TEXT("Core.Cooking"), Attributes);
+#endif
+}
+
+void FStudioTelemetryEditor::RecordEvent_Loading(const FString& LoadingName, double LoadingSeconds, TArray<FAnalyticsEventAttribute> Attributes )
+{
+	const int SchemaVersion = 3;
+
+	Attributes.Emplace(TEXT("SchemaVersion"), SchemaVersion);
+	Attributes.Emplace(TEXT("LoadingName"), LoadingName);
+	Attributes.Emplace(TEXT("LoadingSeconds"), LoadingSeconds);
+	
+#if ENABLE_COOK_STATS
+
+#if UE_WITH_ZEN
+	// Gather Zen analytics
+	if (UE::Zen::IsDefaultServicePresent())
+	{
+		UE::Zen::GetDefaultServiceInstance().GatherAnalytics(Attributes);
+	}
+#endif
+
+	if (UE::Virtualization::IVirtualizationSystem::Get().IsEnabled())
+	{
+		// Gather Virtualization analytics
+		UE::Virtualization::IVirtualizationSystem::Get().GatherAnalytics(Attributes);
+	}
+
+	// Gather the DDC summary stats
+	FDerivedDataCacheSummaryStats SummaryStats;
+
+	GatherDerivedDataCacheSummaryStats(SummaryStats);
+
+	// Append to the attributes
+	for (const FDerivedDataCacheSummaryStat& Stat : SummaryStats.Stats)
+	{
+		FString AttributeName = TEXT("DDC_Summary_") + Stat.Key.Replace(TEXT("."), TEXT("_"));
+
+		if (Stat.Value.IsNumeric())
+		{
+			Attributes.Emplace(AttributeName, FCString::Atof(*Stat.Value));
+		}
+		else
+		{
+			Attributes.Emplace(AttributeName, Stat.Value);
+		}
+	}
+#endif			
+
+	FStudioTelemetry::Get().RecordEvent(TEXT("Core.Loading"), Attributes);
+}
+
+void FStudioTelemetryEditor::RecordEvent_DDCResource(const FString& Context, TArray<FAnalyticsEventAttribute> Attributes)
+{
+#if ENABLE_COOK_STATS
+	// Gather the latest resource stats
+	TArray<FDerivedDataCacheResourceStat> ResourceStats;
+
+	GatherDerivedDataCacheResourceStats(ResourceStats);
+
+	const int SchemaVersion = 3;
+
+	Attributes.Emplace(TEXT("SchemaVersion"), SchemaVersion);
+	Attributes.Emplace(TEXT("LoadingName"), Context);
+	
+	// Send a resource event per asset type
+	for (const FDerivedDataCacheResourceStat& Stat : ResourceStats)
+	{
+		const int64 TotalCount = Stat.BuildCount + Stat.LoadCount;
+		const double TotalTimeSec = Stat.BuildTimeSec + Stat.LoadTimeSec;
+		const int64 TotalSizeMB = Stat.BuildSizeMB + Stat.LoadSizeMB;
+
+		if (Stat.AssetType.IsEmpty() || TotalCount==0)
+		{
+			// Empty asset type or nothing was built or loaded for this type
+			continue;
+		}
+	
+		TArray<FAnalyticsEventAttribute> EventAttributes = Attributes;
+
+		EventAttributes.Emplace(TEXT("AssetType"), Stat.AssetType);
+		EventAttributes.Emplace(TEXT("Load_Count"), Stat.LoadCount);
+		EventAttributes.Emplace(TEXT("Load_TimeSec"), Stat.LoadTimeSec);
+		EventAttributes.Emplace(TEXT("Load_SizeMB"), Stat.LoadSizeMB);
+		EventAttributes.Emplace(TEXT("Build_Count"), Stat.BuildCount);
+		EventAttributes.Emplace(TEXT("Build_TimeSec"), Stat.BuildTimeSec);
+		EventAttributes.Emplace(TEXT("Build_SizeMB"), Stat.BuildSizeMB);
+		EventAttributes.Emplace(TEXT("Total_Count"), TotalCount);
+		EventAttributes.Emplace(TEXT("Total_TimeSec"), TotalTimeSec);
+		EventAttributes.Emplace(TEXT("Total_SizeMB"), TotalSizeMB);
+		EventAttributes.Emplace(TEXT("Efficiency"), double(Stat.LoadCount)/double(TotalCount) );
+		EventAttributes.Emplace(TEXT("Thread_TimeSec"), Stat.GameThreadTimeSec);
+
+		FStudioTelemetry::Get().RecordEvent(TEXT("Core.DDC.Resource"), EventAttributes);
+	}
+#endif			
+}
+
+
+void FStudioTelemetryEditor::RecordEvent_DDCSummary(const FString& Context, TArray<FAnalyticsEventAttribute> Attributes)
+{
+#if ENABLE_COOK_STATS
+	const int SchemaVersion = 3;
+
+	Attributes.Emplace(TEXT("SchemaVersion"), SchemaVersion);
+	Attributes.Emplace(TEXT("LoadingName"), Context);
+
+	// Gather the summary stats
+	FDerivedDataCacheSummaryStats SummaryStats;
+
+	GatherDerivedDataCacheSummaryStats(SummaryStats);
+
+	// Append to the attributes
+	for (const FDerivedDataCacheSummaryStat& Stat : SummaryStats.Stats)
+	{
+		FString AttributeName = Stat.Key.Replace(TEXT("."), TEXT("_"));
+
+		if (Stat.Value.IsNumeric())
+		{
+			Attributes.Emplace(AttributeName, FCString::Atof(*Stat.Value));
+		}
+		else
+		{
+			Attributes.Emplace(AttributeName, Stat.Value);
+		}
+	}
+
+	FStudioTelemetry::Get().RecordEvent(TEXT("Core.DDC.Summary"), Attributes);
+#endif			
+}
+
+void FStudioTelemetryEditor::RecordEvent_IAS(const FString& Context, TArray<FAnalyticsEventAttribute> Attributes)
+{
+	// Gather the summary stats
+	FDerivedDataCacheSummaryStats SummaryStats;
+
+	using namespace UE::IO::IAS;
+	if (FIoStoreOnDemandModule* OnDemandModule = FModuleManager::Get().GetModulePtr<FIoStoreOnDemandModule>("IoStoreOnDemand"))
+	{
+		if (OnDemandModule->IsEnabled())
+		{
+			const int SchemaVersion = 1;
+
+			Attributes.Emplace(TEXT("SchemaVersion"), SchemaVersion);
+			Attributes.Emplace(TEXT("LoadingName"), Context);
+
+			OnDemandModule->ReportAnalytics(Attributes);
+			FStudioTelemetry::Get().RecordEvent(TEXT("Core.IAS"), Attributes);
+		}		
+	}
+}
+
+void FStudioTelemetryEditor::RecordEvent_Zen(const FString& Context, TArray<FAnalyticsEventAttribute> Attributes)
+{
+#if UE_WITH_ZEN
+	// Gather Zen analytics
+	if (UE::Zen::IsDefaultServicePresent())
+	{
+		const int SchemaVersion = 1;
+
+		Attributes.Emplace(TEXT("SchemaVersion"), SchemaVersion);
+		Attributes.Emplace(TEXT("LoadingName"), Context);
+
+		UE::Zen::GetDefaultServiceInstance().GatherAnalytics(Attributes);
+		FStudioTelemetry::Get().RecordEvent(TEXT("Core.Zen"), Attributes);
+	}
+#endif
+}
+
+void FStudioTelemetryEditor::RecordEvent_VirtualAssets(const FString& Context, TArray<FAnalyticsEventAttribute> Attributes)
+{
+	if (UE::Virtualization::IVirtualizationSystem::Get().IsEnabled())
+	{
+		const int SchemaVersion = 1;
+
+		Attributes.Emplace(TEXT("SchemaVersion"), SchemaVersion);
+		Attributes.Emplace(TEXT("LoadingName"), Context);
+
+		// Gather Virtualization analytics
+		UE::Virtualization::IVirtualizationSystem::Get().GatherAnalytics(Attributes);
+
+		FStudioTelemetry::Get().RecordEvent(TEXT("Core.VirtualAssets"), Attributes);
+	}
+}
+
+void FStudioTelemetryEditor::RecordEvent_CoreSystems(const FString& Context, TArray<FAnalyticsEventAttribute> Attributes)
+{
+	FStudioTelemetryEditor::RecordEvent_DDCResource(Context, Attributes);
+	FStudioTelemetryEditor::RecordEvent_DDCSummary(Context, Attributes);
+	FStudioTelemetryEditor::RecordEvent_IAS(Context, Attributes);
+	FStudioTelemetryEditor::RecordEvent_Zen(Context, Attributes);
+	FStudioTelemetryEditor::RecordEvent_VirtualAssets(Context, Attributes);
+}
+
+void FStudioTelemetryEditor::RegisterCollectionWorkflowDelegates(FTelemetryRouter& Router)
+{
+	Router.OnTelemetry<FAssetAddedToCollectionTelemetryEvent>([](const FAssetAddedToCollectionTelemetryEvent& Event)
+	{
+		const int SchemaVersion = 1;
+		
+		FStudioTelemetry::Get().RecordEvent(TEXT("Editor.Collections.AssetsAdded"),
+		{
+			{ TEXT("SchemaVersion"), SchemaVersion},
+			{ TEXT("DurationSec"), Event.DurationSec},
+			{ TEXT("ObjectCount"), Event.NumAdded},
+			{ TEXT("Workflow"), Event.Workflow},
+			{ TEXT("CollectionShareType"), ECollectionShareType::ToString(Event.CollectionShareType)},
+		});
+	});
+
+	Router.OnTelemetry<FAssetRemovedFromCollectionTelemetryEvent>([](const FAssetRemovedFromCollectionTelemetryEvent& Event)
+	{
+		const int SchemaVersion = 1;
+		
+		FStudioTelemetry::Get().RecordEvent(TEXT("Editor.Collections.AssetsRemoved"),
+		{
+			{ TEXT("SchemaVersion"), SchemaVersion},
+			{ TEXT("DurationSec"), Event.DurationSec},
+			{ TEXT("ObjectCount"), Event.NumRemoved},
+			{ TEXT("Workflow"), Event.Workflow},
+			{ TEXT("CollectionShareType"), ECollectionShareType::ToString(Event.CollectionShareType) },
+		});
+	});
+
+	Router.OnTelemetry<FCollectionCreatedTelemetryEvent>([](const FCollectionCreatedTelemetryEvent& Event)
+	{
+		const int SchemaVersion = 1;
+
+		FStudioTelemetry::Get().RecordEvent(TEXT("Editor.Collections.CollectionCreated"),
+		{
+			{ TEXT("SchemaVersion"), SchemaVersion},
+			{ TEXT("DurationSec"), Event.DurationSec},
+			{ TEXT("CollectionShareType"), ECollectionShareType::ToString(Event.CollectionShareType)},
+		});
+	});
+
+	Router.OnTelemetry<FCollectionsDeletedTelemetryEvent>([](const FCollectionsDeletedTelemetryEvent& Event)
+	{
+		const int SchemaVersion = 1;
+
+		FStudioTelemetry::Get().RecordEvent(TEXT("Editor.Collections.CollectionDeleted"),
+		{
+			{ TEXT("SchemaVersion"), SchemaVersion},
+			{ TEXT("DurationSec"), Event.DurationSec},
+			{ TEXT("ObjectCount"), Event.CollectionsDeleted},
+		});
+	});
+}
+
+void FStudioTelemetryEditor::Initialize()
+{
+	SessionStartTime = FPlatformTime::Seconds();
+
+	TSharedPtr<FAnalyticsFlowTracker> FlowTracker = FStudioTelemetry::Get().GetFlowTracker().Pin();
+
+	if (FlowTracker.IsValid())
+	{
+		// Start Editor Boot Flow immediately
+		EditorBootFlowGuid = FlowTracker->StartFlow(TEXT("Editor Boot"));
+		FlowTracker->StartSubFlow(TEXT("Editor Boot"), EditorBootFlowGuid);
+	}
+
+	FEditorDelegates::BeginPIE.AddLambda([this](bool)
+		{
+			TSharedPtr<FAnalyticsFlowTracker> FlowTracker = FStudioTelemetry::Get().GetFlowTracker().Pin();
+
+			if (FlowTracker.IsValid())
+			{
+				// Finish the Editor Interaction Flow
+				FlowTracker->EndFlow(InteractiveEditorFlowGuid);
+
+				// Start PIE Flow
+				PIEStartTime = FPlatformTime::Seconds();
+				PIEFlowGuid = FlowTracker->StartFlow(TEXT("Play In Editor"));
+				PIEInitializeSubFlowGuid = FlowTracker->StartSubFlow(TEXT("PIE Initialize"), PIEFlowGuid);
+			}			
+		});
+
+	FEditorDelegates::PostPIEStarted.AddLambda([this](bool)
+		{
+			TSharedPtr<FAnalyticsFlowTracker> FlowTracker = FStudioTelemetry::Get().GetFlowTracker().Pin();
+
+			if (FlowTracker.IsValid())
+			{
+				// Called when PIE has finally started
+				FlowTracker->EndSubFlow(PIEInitializeSubFlowGuid);
+			}
+		});
+
+	FEditorDelegates::EndPIE.AddLambda([this](bool)
+		{
+			TSharedPtr<FAnalyticsFlowTracker> FlowTracker = FStudioTelemetry::Get().GetFlowTracker().Pin();
+
+			if (FlowTracker.IsValid())
+			{
+				// End PIE Flow and any SubFlows within it
+				FlowTracker->EndFlow(PIEFlowGuid, true);
+
+				// Restart the Interactive Editor Flow
+				FlowTracker->StartFlow(TEXT("Interactive Editor"));
+			}
+		});
+
+	FWorldDelegates::OnPIEStarted.AddLambda([this](UGameInstance* GameInstance)
+		{
+			TSharedPtr<FAnalyticsFlowTracker> FlowTracker = FStudioTelemetry::Get().GetFlowTracker().Pin();
+
+			if (FlowTracker.IsValid())
+			{
+				PIELoadMapStartTime = FPlatformTime::Seconds();
+				PIELoadMapSubFlowGuid = FlowTracker->StartSubFlow(TEXT("PIE Load Map"), PIEFlowGuid);
+			}
+		});
+
+	FWorldDelegates::OnPIEReady.AddLambda([this](UGameInstance* GameInstance)
+		{	
+			const FString MapName = FPaths::GetBaseFilename(GameInstance->PIEMapName);
+
+			TArray<FAnalyticsEventAttribute> Attributes;
+			Attributes.Emplace(FAnalyticsEventAttribute(TEXT("LevelName"), LevelName));
+			Attributes.Emplace(FAnalyticsEventAttribute(TEXT("MapName"), MapName));
+			
+			// Currently there is no obvious way to know when PIE has finished loading everything but loading map FrontEnd is good enough for Epic products.
+			if (MapName.Contains(TEXT("Frontend")))
+			{
+				static bool IsFirstTimeToPIE = true;
+
+				PIEStartupTime = FPlatformTime::Seconds() - PIEStartTime;
+
+				if (IsFirstTimeToPIE == true)
+				{
+					// Record the absolute time from editor boot to PIE
+					FStudioTelemetryEditor::RecordEvent_Loading(TEXT("TimeToPIE"), EditorStartupTime + LoadMapTime + PIEStartupTime, Attributes);
+					FStudioTelemetryEditor::RecordEvent_CoreSystems(TEXT("TimeToPIE"), Attributes);
+
+					IsFirstTimeToPIE = false;
+				}
+				
+				// Record the time from start PIE to PIE
+				FStudioTelemetryEditor::RecordEvent_Loading(TEXT("PIE.TotalStartupTime"), PIEStartupTime, Attributes);
+				FStudioTelemetryEditor::RecordEvent_CoreSystems(TEXT("PIE.TotalStartupTime"), Attributes);
+			}
+
+			PIELoadMapTime = FPlatformTime::Seconds() - PIELoadMapStartTime;
+
+			// Record the sub map load time
+			FStudioTelemetryEditor::RecordEvent_Loading(TEXT("PIE.LoadMapTime"), PIELoadMapTime, Attributes);
+			FStudioTelemetryEditor::RecordEvent_CoreSystems(TEXT("PIE.LoadMapTime"), Attributes);
+
+			TSharedPtr<FAnalyticsFlowTracker> FlowTracker = FStudioTelemetry::Get().GetFlowTracker().Pin();
+
+			if (FlowTracker.IsValid())
+			{
+				FlowTracker->EndSubFlow(PIELoadMapSubFlowGuid, true, Attributes);
+
+				// Start tracking World Streaming
+				if (UWorld* World = GameInstance->GetWorld())
+				{
+					FGuid WorldStreamingSubFlowGuid = FlowTracker->StartSubFlow(TEXT("World Streaming"));
+
+					WorldStreamingStartTime = FPlatformTime::Seconds();
+
+					World->OnWorldMatchStarting.AddLambda([this, WorldStreamingSubFlowGuid, MapName, FlowTracker]()
+						{
+							TArray<FAnalyticsEventAttribute> Attributes;
+							Attributes.Emplace(FAnalyticsEventAttribute(TEXT("LevelName"), LevelName));
+							Attributes.Emplace(FAnalyticsEventAttribute(TEXT("MapName"), MapName));
+
+							FlowTracker->EndSubFlow(WorldStreamingSubFlowGuid, true, Attributes);
+							
+							FStudioTelemetryEditor::RecordEvent_Loading(TEXT("WorldStreaming"), FPlatformTime::Seconds() - WorldStreamingStartTime, Attributes);
+						});
+				}
+			}
+		});
+
+	FEditorFileUtils::GetOnLoadMapStartDelegate().AddLambda([this]()
+		{
+			TSharedPtr<FAnalyticsFlowTracker> FlowTracker = FStudioTelemetry::Get().GetFlowTracker().Pin();
+
+			if (FlowTracker.IsValid())
+			{
+				// Load Map SubFlow will inherit the current Flow scope
+				LoadMapStartTime = FPlatformTime::Seconds();
+				LoadMapSubFlowGuid = FlowTracker->StartSubFlow(TEXT("Load Map"));
+			}
+		});
+
+	FEditorFileUtils::GetOnLoadMapEndDelegate().AddLambda([this](const FString& MapName)
+		{
+			LevelName = MapName;
+
+			TArray<FAnalyticsEventAttribute> Attributes;
+			Attributes.Emplace(FAnalyticsEventAttribute(TEXT("LevelName"), LevelName));
+			Attributes.Emplace(FAnalyticsEventAttribute(TEXT("MapName"), MapName));
+
+			LoadMapTime = FPlatformTime::Seconds() - LoadMapStartTime;
+			
+			FStudioTelemetryEditor::RecordEvent_Loading(TEXT("LoadMap"), LoadMapTime, Attributes);
+			FStudioTelemetryEditor::RecordEvent_CoreSystems(TEXT("LoadMap"), Attributes);
+
+			TSharedPtr<FAnalyticsFlowTracker> FlowTracker = FStudioTelemetry::Get().GetFlowTracker().Pin();
+
+			if (FlowTracker.IsValid())
+			{
+				FlowTracker->EndSubFlow(LoadMapSubFlowGuid, true, Attributes);
+			}
+		});
+
+	FEditorDelegates::OnEditorBoot.AddLambda([this](double TimeToBootEditor)
+		{
+			// Editor Boot Flow has finished
+			TArray<FAnalyticsEventAttribute> Attributes;
+			Attributes.Emplace(FAnalyticsEventAttribute(TEXT("LevelName"), TEXT("EditorBoot")));
+			Attributes.Emplace(FAnalyticsEventAttribute(TEXT("MapName"), TEXT("EditorBoot")));
+			
+			FStudioTelemetryEditor::RecordEvent_Loading(TEXT("BootEditor"), TimeToBootEditor, Attributes);
+			FStudioTelemetryEditor::RecordEvent_CoreSystems(TEXT("BootEditor"), Attributes);
+
+			TSharedPtr<FAnalyticsFlowTracker> FlowTracker = FStudioTelemetry::Get().GetFlowTracker().Pin();
+
+			if (FlowTracker.IsValid())
+			{
+				FlowTracker->EndFlow(EditorBootFlowGuid, true, Attributes);
+
+				// Start the Editor Interact Flow
+				InteractiveEditorFlowGuid = FlowTracker->StartFlow(TEXT("Interactive Editor"));
+			}
+		});
+
+	FEditorDelegates::OnEditorInitialized.AddLambda([this](double TimeToInitializeEditor)
+		{
+			EditorStartupTime = TimeToInitializeEditor;
+
+			// Editor has initialized
+			TArray<FAnalyticsEventAttribute> Attributes;
+			Attributes.Emplace(FAnalyticsEventAttribute(TEXT("LevelName"), TEXT("EditorInitialize")));
+			Attributes.Emplace(FAnalyticsEventAttribute(TEXT("MapName"), TEXT("EditorInitialize")));
+			
+			FStudioTelemetryEditor::RecordEvent_Loading(TEXT("TotalEditorStartup"), TimeToInitializeEditor, Attributes);
+			FStudioTelemetryEditor::RecordEvent_CoreSystems(TEXT("TotalEditorStartup"), Attributes);
+
+			ensureMsgf(GEditor, TEXT("GEditor was not valid"));
+
+			if (GEditor != nullptr)
+			{
+				// Install callbacks for Open Asset Dialogue
+				GEditor->GetEditorSubsystem<UAssetEditorSubsystem>()->OnAssetEditorRequestedOpen().AddLambda([this](UObject* Asset)
+					{
+						AssetOpenStartTime = FPlatformTime::Seconds();
+						TSharedPtr<FAnalyticsFlowTracker> FlowTracker = FStudioTelemetry::Get().GetFlowTracker().Pin();
+
+						if (FlowTracker.IsValid())
+						{
+							FlowTracker->StartSubFlow(TEXT("Open Asset Editor"));
+						}
+					});
+
+				GEditor->GetEditorSubsystem<UAssetEditorSubsystem>()->OnAssetOpenedInEditor().AddLambda([this](UObject* Asset, IAssetEditorInstance*)
+					{
+						TArray<FAnalyticsEventAttribute> Attributes;
+						Attributes.Emplace(FAnalyticsEventAttribute(TEXT("LevelName"), LevelName));
+						
+						if (Asset != nullptr)
+						{
+							Attributes.Emplace(FAnalyticsEventAttribute(TEXT("MapName"), Asset->GetFullName()));
+							Attributes.Emplace(TEXT("AssetPath"), Asset->GetFullName());
+							Attributes.Emplace(TEXT("AssetType"), Asset->GetClass()->GetName());
+
+							FStudioTelemetryEditor::RecordEvent_Loading(TEXT("OpenAssetEditor"), FPlatformTime::Seconds() - AssetOpenStartTime, Attributes);
+						}
+
+						TSharedPtr<FAnalyticsFlowTracker> FlowTracker = FStudioTelemetry::Get().GetFlowTracker().Pin();
+
+						if (FlowTracker.IsValid())
+						{
+							FlowTracker->EndSubFlow(TEXT("Open Asset Editor"), true, Attributes);
+						}
+				});
+			}
+
+			ensureMsgf(GUnrealEd, TEXT("GUnrealEd was not valid"));
+			
+			if (GUnrealEd!=nullptr && GUnrealEd->CookServer != nullptr )
+			{
+				UE::Cook::FDelegates::CookByTheBookFinished.AddLambda([this](UE::Cook::ICookInfo& CookInfo)
+					{
+						TArray<FAnalyticsEventAttribute> Attributes;
+						Attributes.Emplace(FAnalyticsEventAttribute(TEXT("LevelName"), LevelName));
+						Attributes.Emplace(FAnalyticsEventAttribute(TEXT("MapName"), LevelName));
+						
+						FStudioTelemetryEditor::RecordEvent_Cooking(Attributes);
+						FStudioTelemetryEditor::RecordEvent_CoreSystems(TEXT("Cooking"), Attributes);
+					});
+			}
+
+			ensureMsgf(GWarn, TEXT("GWarn was not valid"));
+
+			if (GWarn != nullptr)
+			{
+				GWarn->OnStartSlowTask().AddLambda([this](const FText& TaskName)
+					{	
+						TSharedPtr<FAnalyticsFlowTracker> FlowTracker = FStudioTelemetry::Get().GetFlowTracker().Pin();
+
+						if (FlowTracker.IsValid())
+						{
+							FlowTracker->StartSubFlow(FName(TaskName.ToString()));
+						}
+					});
+
+				GWarn->OnFinalizeSlowTask().AddLambda([this](const FText& TaskName, double TaskDurationSeconds)
+					{
+						TArray<FAnalyticsEventAttribute> Attributes;
+						Attributes.Emplace(FAnalyticsEventAttribute(TEXT("LevelName"), LevelName));
+						Attributes.Emplace(TEXT("TaskName"), TaskName.ToString());
+						
+						TSharedPtr<FAnalyticsFlowTracker> FlowTracker = FStudioTelemetry::Get().GetFlowTracker().Pin();
+
+						if (FlowTracker.IsValid())
+						{
+							FlowTracker->EndSubFlow(FName(TaskName.ToString()), true, Attributes);
+						}
+						
+						FStudioTelemetryEditor::RecordEvent_Loading(TEXT("SlowTaskDialog"), TaskDurationSeconds, Attributes);
+					});
+			}
+		});
+	
+	// Install Cooking Callbacks
+	UE::Cook::FDelegates::CookByTheBookStarted.AddLambda([this](UE::Cook::ICookInfo& CookInfo)
+	{
+			TSharedPtr<FAnalyticsFlowTracker> FlowTracker = FStudioTelemetry::Get().GetFlowTracker().Pin();
+
+			if (FlowTracker.IsValid())
+			{
+				CookByTheBookFlowGuid = FlowTracker->StartFlow("Cook By The Book");
+			}
+	});
+
+	UE::Cook::FDelegates::CookByTheBookFinished.AddLambda([this](UE::Cook::ICookInfo& CookInfo)
+	{
+		// Suppress sending telemetry from CookWorkers for now.
+		uint32 MultiprocessId = 0;
+		FParse::Value(FCommandLine::Get(), TEXT("-MultiprocessId="), MultiprocessId);
+		if (MultiprocessId != 0)
+		{
+			return;
+		}
+
+		TArray<FAnalyticsEventAttribute> Attributes;
+		Attributes.Emplace(FAnalyticsEventAttribute(TEXT("LevelName"), LevelName));
+		Attributes.Emplace(FAnalyticsEventAttribute(TEXT("MapName"), LevelName));
+
+		TSharedPtr<FAnalyticsFlowTracker> FlowTracker = FStudioTelemetry::Get().GetFlowTracker().Pin();
+
+		if (FlowTracker.IsValid())
+		{
+			FlowTracker->EndFlow(CookByTheBookFlowGuid, true, Attributes);
+		}
+
+		FStudioTelemetryEditor::RecordEvent_Cooking(Attributes);
+		FStudioTelemetryEditor::RecordEvent_CoreSystems(TEXT("Cooking"), Attributes);
+	});
+
+	FContentBrowserModule* ContentBrowserModule = FModuleManager::GetModulePtr<FContentBrowserModule>( TEXT("ContentBrowser") );
+	
+	FTelemetryRouter& Router = FTelemetryRouter::Get();
+	{
+		using namespace UE::Telemetry::ContentBrowser;
+		Router.OnTelemetry<FBackendFilterTelemetry>([this](const FBackendFilterTelemetry& Data)
+		{
+			FString DataFilterText = LexToString(FJsonNull{});
+			if (Data.DataFilter)
+			{
+				Private::FAnalyticsJsonWriter J(&DataFilterText);
+				J.WriteObjectStart();
+				J.WriteValue("RecursivePaths", Data.DataFilter->bRecursivePaths);
+				J.WriteValue("ItemTypeFilter", UEnum::GetValueOrBitfieldAsString(Data.DataFilter->ItemTypeFilter));
+				J.WriteValue("ItemCategoryFilter", UEnum::GetValueOrBitfieldAsString(Data.DataFilter->ItemCategoryFilter));
+				J.WriteValue("ItemAttributeFilter", UEnum::GetValueOrBitfieldAsString(Data.DataFilter->ItemAttributeFilter));
+				TArray<const UScriptStruct*> FilterTypes = Data.DataFilter->ExtraFilters.GetFilterTypes();
+				if (FilterTypes.Num() > 0)
+				{
+					J.WriteArrayStart("FilterTypes");
+					for (const UScriptStruct* Type : FilterTypes)
+					{
+						J.WriteValue(Type->GetPathName());
+					}
+					J.WriteArrayEnd();				
+				}
+				J.WriteObjectEnd();
+				J.Close();
+			}
+			
+			FStudioTelemetry::Get().RecordEvent(TEXT("Editor.AssetView.BackendFilter"), 
+			{
+				{ TEXT("SchemaVersion"), 1 },
+				{ TEXT("ViewCorrelationGuid"), Data.ViewCorrelationGuid },
+				{ TEXT("FilterSessionCorrelationGuid"), Data.FilterSessionCorrelationGuid },
+				{ TEXT("HasCustomItemSources"), Data.bHasCustomItemSources },
+				{ TEXT("RefreshSourceItemsDurationSeconds"), Data.RefreshSourceItemsDurationSeconds },
+				{ TEXT("NumBackendItems"), Data.NumBackendItems },
+				{ TEXT("DataFilter"), FJsonFragment(MoveTemp(DataFilterText)) },
+			});
+		});
+
+		Router.OnTelemetry<FFrontendFilterTelemetry>([this](const FFrontendFilterTelemetry& Data)
+		{
+			FString FilterText = LexToString(FJsonNull{});
+			if (Data.FrontendFilters.IsValid() && Data.FrontendFilters->Num())
+			{
+				Private::FAnalyticsJsonWriter J(&FilterText);
+				J.WriteArrayStart();
+				for (int32 i=0; i < Data.FrontendFilters->Num(); ++i)
+				{
+					TSharedPtr<IFilter<FAssetFilterType>> Filter = Data.FrontendFilters->GetFilterAtIndex(i);
+					J.WriteValue(Filter->GetName());
+				}
+				J.WriteArrayEnd();
+				J.Close();
+			}
+			FStudioTelemetry::Get().RecordEvent(TEXT("Editor.AssetView.FrontendFilter"), 
+			{
+				{ TEXT("SchemaVersion"), 1 },
+				{ TEXT("ViewCorrelationGuid"), Data.ViewCorrelationGuid },
+				{ TEXT("FilterSessionCorrelationGuid"), Data.FilterSessionCorrelationGuid },
+				{ TEXT("TotalItemsToFilter"), Data.TotalItemsToFilter },
+				{ TEXT("PriorityItemsToFilter"), Data.PriorityItemsToFilter },
+				{ TEXT("TotalResults"), Data.TotalResults },
+				{ TEXT("AmortizeDurationSeconds"), Data.AmortizeDuration },
+				{ TEXT("WorkDurationSeconds"), Data.WorkDuration },
+				{ TEXT("ResultLatency"), AnalyticsOptionalToStringOrNull(Data.ResultLatency) },
+				{ TEXT("TimeUntilInteractionSeconds"), AnalyticsOptionalToStringOrNull(Data.TimeUntilInteraction) },
+				{ TEXT("Completed"), Data.bCompleted },
+				{ TEXT("FrontendFilters"), FJsonFragment(MoveTemp(FilterText)) },
+			});
+		});
+
+		RegisterCollectionWorkflowDelegates(Router);
+	}
+	{
+		using namespace UE::Telemetry::AssetRegistry;
+
+		Router.OnTelemetry<FStartupTelemetry>([this](const FStartupTelemetry& Data){
+			FStudioTelemetry::Get().RecordEvent(TEXT("Editor.AssetRegistry.Startup"), 
+			{
+				{ TEXT("SchemaVersion"), 1 },
+				{ TEXT("Duration"), Data.StartupDuration },
+				{ TEXT("StartedAsyncGather"), Data.bStartedAsyncGather },
+			});
+		});
+		Router.OnTelemetry<FSynchronousScanTelemetry>([this](const FSynchronousScanTelemetry& Data){
+			if (Data.Duration < 0.5)
+			{
+				return;
+			}
+			FString DirectoriesText;
+			{
+				Private::FAnalyticsJsonWriter J(&DirectoriesText);
+				J.WriteArrayStart();
+				for (const FString& Directory : MakeArrayView(Data.Directories).Left(100))
+				{
+					J.WriteValue(Directory);
+				}
+				J.WriteArrayEnd();
+				J.Close();
+			}
+			FString FilesText;
+			{
+				Private::FAnalyticsJsonWriter J(&FilesText);
+				J.WriteArrayStart();
+				for (const FString& File : MakeArrayView(Data.Files).Left(100))
+				{
+					J.WriteValue(File);
+				}
+				J.WriteArrayEnd();
+				J.Close();
+			}
+			FStudioTelemetry::Get().RecordEvent(TEXT("Editor.AssetRegistry.SynchronousScan"), 
+			{
+				{ TEXT("SchemaVersion"), 1 },
+				{ TEXT("Directories"), FJsonFragment(MoveTemp(DirectoriesText)) },
+				{ TEXT("Files"), FJsonFragment(MoveTemp(FilesText)) },
+				{ TEXT("Flags"), LexToString(Data.Flags) },
+				{ TEXT("NumFoundAssets"), Data.NumFoundAssets },
+				{ TEXT("DurationSeconds"), Data.Duration },
+				{ TEXT("InitialSearchStarted"), Data.bInitialSearchStarted },
+				{ TEXT("InitialSearchCompleted"), Data.bInitialSearchCompleted },
+			});
+		});
+		Router.OnTelemetry<FGatherTelemetry>([this](const FGatherTelemetry& Data){
+			FStudioTelemetry::Get().RecordEvent(TEXT("Editor.AssetRegistry.InitialScan"), 
+			{
+				{ TEXT("SchemaVersion"), 1 },
+				{ TEXT("TotalDurationSeconds"), Data.TotalSearchDurationSeconds },
+				{ TEXT("TotalWorkSeconds"), Data.TotalWorkTimeSeconds },
+				{ TEXT("DiscoverySeconds"), Data.DiscoveryTimeSeconds },
+				{ TEXT("GatherSeconds"), Data.GatherTimeSeconds },
+				{ TEXT("StoreSeconds"), Data.StoreTimeSeconds },
+				{ TEXT("NumCachedDirectories"), Data.NumCachedDirectories },
+				{ TEXT("NumUncachedDirectories"), Data.NumUncachedDirectories },
+				{ TEXT("NumCachedAssetFiles"), Data.NumCachedAssetFiles },
+				{ TEXT("NumUncachedAssetFiles"), Data.NumUncachedAssetFiles },
+			});
+		});
+		Router.OnTelemetry<FDirectoryWatcherUpdateTelemetry>([this](const FDirectoryWatcherUpdateTelemetry& Data){
+			if (Data.DurationSeconds < 0.5)
+			{
+				return;
+			}
+			FStudioTelemetry::Get().RecordEvent(TEXT("Editor.AssetRegistry.DirectoryWatcherUpdate"), 
+			{
+				{ TEXT("SchemaVersion"), 1 },
+				{ TEXT("NumChanges"), Data.Changes.Num() },
+				{ TEXT("DurationSeconds"), Data.DurationSeconds },
+				{ TEXT("InitialSearchStarted"), Data.bInitialSearchStarted },
+				{ TEXT("InitialSearchCompleted"), Data.bInitialSearchCompleted },
+			});
+		});
+	}
+}
+
+void FStudioTelemetryEditor::Shutdown()
+{
+}
+
+UE_ENABLE_OPTIMIZATION_SHIP
+
+#endif // WITH_EDITOR
\ No newline at end of file
diff --git a/Engine/Plugins/Experimental/StudioTelemetry/Source/StudioTelemetry/Private/StudioTelemetryEditor.h b/Engine/Plugins/Experimental/StudioTelemetry/Source/StudioTelemetry/Private/StudioTelemetryEditor.h
new file mode 100644
--- /dev/null
+++ b/Engine/Plugins/Experimental/StudioTelemetry/Source/StudioTelemetry/Private/StudioTelemetryEditor.h
@@ -0,0 +1,60 @@
+// Copyright Epic Games, Inc. All Rights Reserved.
+
+#pragma once
+
+#if WITH_EDITOR
+
+#include "StudioTelemetry.h"
+
+class FTelemetryRouter;
+
+/**
+ * A class that implements a variety of pre-configured Core and Editor telemetry events that can be used to evaluate the efficiecnty of the most common developer workflows
+ */
+class FStudioTelemetryEditor : FNoncopyable
+{
+public:
+	FStudioTelemetryEditor() {};
+	~FStudioTelemetryEditor() {};
+
+	static FStudioTelemetryEditor& Get();
+	void Initialize();
+	void Shutdown();
+
+private:	
+
+	/** Various event handling functions */
+	static void RecordEvent_Cooking(TArray<FAnalyticsEventAttribute> Attributes = {});
+	static void RecordEvent_Loading(const FString& LoadingName, double LoadingSeconds, TArray<FAnalyticsEventAttribute> Attributes = {});
+	static void RecordEvent_CoreSystems(const FString& Context, TArray<FAnalyticsEventAttribute> Attributes = {});
+	static void RecordEvent_DDCResource(const FString& Context, TArray<FAnalyticsEventAttribute> Attributes = {});
+	static void RecordEvent_DDCSummary(const FString& Context, TArray<FAnalyticsEventAttribute> Attributes = {});
+	static void RecordEvent_IAS(const FString& Context, TArray<FAnalyticsEventAttribute> Attributes = {});
+	static void RecordEvent_Zen(const FString& Context, TArray<FAnalyticsEventAttribute> Attributes = {});
+	static void RecordEvent_VirtualAssets(const FString& Context, TArray<FAnalyticsEventAttribute> Attributes = {});
+	static void RegisterCollectionWorkflowDelegates(FTelemetryRouter& Router);
+
+	/** Internal variables for Flows and Timer data */
+	FString LevelName;
+	FGuid EditorBootFlowGuid;
+	FGuid InteractiveEditorFlowGuid;
+	FGuid LoadMapSubFlowGuid;
+	FGuid PIEFlowGuid;
+	FGuid PIEInitializeSubFlowGuid;
+	FGuid PIELoadMapSubFlowGuid;
+	FGuid CookByTheBookFlowGuid;
+
+	double SessionStartTime;
+	double PIEStartTime;
+	double LoadMapStartTime;
+	double AssetOpenStartTime;
+	double WorldStreamingStartTime;
+	double EditorStartupTime=0;
+	double LoadMapTime=0;
+	double PIEStartupTime = 0;
+	double PIELoadMapTime = 0;
+	double PIELoadMapStartTime;
+};
+
+#endif // WITH_EDITOR
+
diff --git a/Engine/Plugins/Experimental/StudioTelemetry/Source/StudioTelemetry/Private/StudioTelemetryLog.h b/Engine/Plugins/Experimental/StudioTelemetry/Source/StudioTelemetry/Private/StudioTelemetryLog.h
new file mode 100644
--- /dev/null
+++ b/Engine/Plugins/Experimental/StudioTelemetry/Source/StudioTelemetry/Private/StudioTelemetryLog.h
@@ -0,0 +1,10 @@
+// Copyright Epic Games, Inc. All Rights Reserved.
+
+#pragma once
+
+#include "CoreMinimal.h"
+
+/**
+ * StudioTelemetry plugin logging
+ */
+DECLARE_LOG_CATEGORY_EXTERN(LogStudioTelemetry, Log, All);
diff --git a/Engine/Plugins/Experimental/StudioTelemetry/Source/StudioTelemetry/Public/StudioTelemetry.h b/Engine/Plugins/Experimental/StudioTelemetry/Source/StudioTelemetry/Public/StudioTelemetry.h
new file mode 100644
--- /dev/null
+++ b/Engine/Plugins/Experimental/StudioTelemetry/Source/StudioTelemetry/Public/StudioTelemetry.h
@@ -0,0 +1,64 @@
+// Copyright Epic Games, Inc. All Rights Reserved.
+
+#pragma once
+
+#include "Interfaces/IAnalyticsProvider.h"
+#include "AnalyticsFlowTracker.h"
+#include "Containers/Array.h"
+#include "Containers/UnrealString.h"
+#include "CoreMinimal.h"
+#include "Modules/ModuleManager.h"
+#include "Templates/SharedPointer.h"
+
+class FAnalyticsProviderMulticast;
+
+/**
+ * Public facing Studio Telemetry Plugin API
+ */
+class FStudioTelemetry : public IModuleInterface
+{
+public:
+
+	typedef TFunction<void(const FString& EventName, const TArray<FAnalyticsEventAttribute>& Attrs)> OnRecordEvent;
+
+	/** Check whether the module is available*/
+	static STUDIOTELEMETRY_API bool IsAvaliable() { return FModuleManager::Get().IsModuleLoaded("StudioTelemetry"); }
+
+	/** Access to the module singleton*/
+	static STUDIOTELEMETRY_API FStudioTelemetry& Get();
+
+	/** Access to the a specific named analytics provider within the system*/
+	STUDIOTELEMETRY_API TWeakPtr<IAnalyticsProvider> GetProvider(const FString& ProviderName);
+
+	/** Access to the broadcast analytics provider for the system*/
+	STUDIOTELEMETRY_API TWeakPtr<IAnalyticsProvider> GetProvider();
+
+	/** Access to the flow tracker for the system, this will ultimately broadcast flow events to all providers*/
+	STUDIOTELEMETRY_API TWeakPtr<FAnalyticsFlowTracker> GetFlowTracker();
+	
+	/** Thread safe method to record an event to all analytics providers*/
+	STUDIOTELEMETRY_API void RecordEvent(const FString& EventName, const TArray<FAnalyticsEventAttribute>& Attributes = {});
+
+	/** Thread safe method to record an event to the named analytics provider */
+	STUDIOTELEMETRY_API void RecordEvent(const FString& ProviderName, const FString& EventName, const TArray<FAnalyticsEventAttribute>& Attributes = {});
+	
+	/** Method for custom recording of telemetry events*/
+	STUDIOTELEMETRY_API void SetRecordEventCallback(OnRecordEvent);
+
+private:
+
+	/** IModuleInterface implementation */
+	STUDIOTELEMETRY_API virtual void StartupModule()  final;
+	STUDIOTELEMETRY_API virtual void ShutdownModule()  final;
+
+	/** Starts a new analytics session*/
+	void StartSession();
+
+	/** Ends an existing analytics session*/
+	void EndSession();
+
+	FCriticalSection						CriticalSection;
+	TSharedPtr<FAnalyticsProviderMulticast>	AnalyticsProvider;
+	TSharedPtr<FAnalyticsFlowTracker>		AnalyticsFlowTracker;
+	OnRecordEvent							RecordEventCallback;
+};
diff --git a/Engine/Plugins/Experimental/StudioTelemetry/Source/StudioTelemetry/StudioTelemetry.build.cs b/Engine/Plugins/Experimental/StudioTelemetry/Source/StudioTelemetry/StudioTelemetry.build.cs
new file mode 100644
--- /dev/null
+++ b/Engine/Plugins/Experimental/StudioTelemetry/Source/StudioTelemetry/StudioTelemetry.build.cs
@@ -0,0 +1,56 @@
+// Copyright Epic Games, Inc. All Rights Reserved.
+
+using System.IO;
+using UnrealBuildTool;
+
+public class StudioTelemetry : ModuleRules
+{
+	public StudioTelemetry(ReadOnlyTargetRules Target) : base(Target)
+	{
+		PCHUsage = ModuleRules.PCHUsageMode.UseExplicitOrSharedPCHs;
+		
+		PublicDependencyModuleNames.AddRange(
+			new string[] {
+				"Core",
+				"HTTP"
+			}
+		);
+
+		PrivateDependencyModuleNames.AddRange(
+			new string[] {
+				"CoreUObject",
+				"Engine",
+                "EngineSettings",
+                "BuildSettings",
+                "Analytics",
+                "AnalyticsET",
+				"TelemetryUtils",
+				"RHI"
+			}
+		);
+
+		if (Target.bBuildEditor)
+		{
+			PrivateDependencyModuleNames.AddRange(
+				new string[] {
+					"Horde",
+					"AssetRegistry",
+					"UnrealEd",
+					"DerivedDataCache",
+					"Zen",
+					"IoStoreOnDemand",
+					"ContentBrowser",
+					"ContentBrowserData",
+					"TelemetryUtils"
+				}
+			);
+
+			PublicDependencyModuleNames.AddRange(
+				new string[] {
+					"AnalyticsET",
+					"AnalyticsLog"
+				}
+			);
+		}
+	}
+}
diff --git a/Engine/Plugins/Experimental/StudioTelemetry/StudioTelemetry.uplugin b/Engine/Plugins/Experimental/StudioTelemetry/StudioTelemetry.uplugin
--- a/Engine/Plugins/Experimental/StudioTelemetry/StudioTelemetry.uplugin
+++ b/Engine/Plugins/Experimental/StudioTelemetry/StudioTelemetry.uplugin
@@ -15,6 +15,11 @@
 	"IsBetaVersion": false,
 	"Installed": false,
 	"Modules": [
+		{
+			"Name": "AnalyticsLog",
+			"Type": "EditorAndProgram",
+			"LoadingPhase": "PostConfigInit"
+		},
 		{
 			"Name": "StudioTelemetry",
 			"Type": "EditorAndProgram",
diff --git a/Engine/Plugins/Runtime/Analytics/AnalyticsMulticast/Source/AnalyticsMulticast/Private/AnalyticsMulticast.cpp b/Engine/Plugins/Runtime/Analytics/AnalyticsMulticast/Source/AnalyticsMulticast/Private/AnalyticsMulticast.cpp
--- a/Engine/Plugins/Runtime/Analytics/AnalyticsMulticast/Source/AnalyticsMulticast/Private/AnalyticsMulticast.cpp
+++ b/Engine/Plugins/Runtime/Analytics/AnalyticsMulticast/Source/AnalyticsMulticast/Private/AnalyticsMulticast.cpp
@@ -42,6 +42,11 @@ public:
 	virtual void RecordCurrencyPurchase(const FString& GameCurrencyType, int GameCurrencyAmount, const FString& RealCurrencyType, float RealMoneyCost, const FString& PaymentProvider) override;
 	virtual void RecordCurrencyGiven(const FString& GameCurrencyType, int GameCurrencyAmount) override;
 
+	virtual void SetDefaultEventAttributes(TArray<FAnalyticsEventAttribute>&& Attributes)  override;
+	virtual TArray<FAnalyticsEventAttribute> GetDefaultEventAttributesSafe() const override;
+	virtual int32 GetDefaultEventAttributeCount() const  override;
+	virtual FAnalyticsEventAttribute GetDefaultEventAttribute(int AttributeIndex) const  override;
+
 	virtual ~FAnalyticsProviderMulticast();
 
 	bool HasValidProviders() const { return Providers.Num() > 0; }
@@ -61,6 +66,7 @@ public:
 private:
 	TArray<TSharedPtr<IAnalyticsProvider> > Providers;
 	TArray<FString> ProviderModules;
+	TArray<FAnalyticsEventAttribute> DefaultEventAttributes;
 };
 
 TSharedPtr<IAnalyticsProvider> FAnalyticsProviderMulticast::Provider;
@@ -322,3 +328,28 @@ void FAnalyticsProviderMulticast::RecordProgress(const FString& ProgressType, co
 	}
 }
 
+void FAnalyticsProviderMulticast::SetDefaultEventAttributes(TArray<FAnalyticsEventAttribute>&& Attributes)
+{
+	DefaultEventAttributes = Attributes;
+
+	for (auto CurrentProvider : Providers)
+	{
+		CurrentProvider->SetDefaultEventAttributes(CopyTemp(DefaultEventAttributes));
+	}
+}
+
+TArray<FAnalyticsEventAttribute> FAnalyticsProviderMulticast::GetDefaultEventAttributesSafe() const
+{
+	return DefaultEventAttributes;
+}
+
+int32 FAnalyticsProviderMulticast::GetDefaultEventAttributeCount() const
+{
+	return DefaultEventAttributes.Num();
+}
+
+FAnalyticsEventAttribute FAnalyticsProviderMulticast::GetDefaultEventAttribute(int AttributeIndex) const
+{
+	return DefaultEventAttributes[AttributeIndex];
+}
+
diff --git a/Engine/Plugins/Runtime/Analytics/FileLogging/Source/FileLogging/Private/FileLogging.cpp b/Engine/Plugins/Runtime/Analytics/FileLogging/Source/FileLogging/Private/FileLogging.cpp
--- a/Engine/Plugins/Runtime/Analytics/FileLogging/Source/FileLogging/Private/FileLogging.cpp
+++ b/Engine/Plugins/Runtime/Analytics/FileLogging/Source/FileLogging/Private/FileLogging.cpp
@@ -173,6 +173,26 @@ void FAnalyticsProviderFileLogging::RecordEvent(const FString& EventName, const
 			FileArchive->Logf(TEXT(",\t\t\t\"attributes\" : ["));
 			bool bHasWrittenFirstAttr = false;
 			// Write out the list of attributes as an array of attribute objects
+			for (auto Attr : DefaultEventAttributes)
+			{
+				if (bHasWrittenFirstAttr)
+				{
+					FileArchive->Logf(TEXT("\t\t\t,"));
+				}
+				FileArchive->Logf(TEXT("\t\t\t{"));
+				FileArchive->Logf(TEXT("\t\t\t\t\"name\" : \"%s\","), *Attr.GetName());
+				if (Attr.IsJsonFragment())
+				{
+					FileArchive->Logf(TEXT("\t\t\t\t\"value\" : %s"), *Attr.GetValue());
+				}
+				else
+				{
+					FileArchive->Logf(TEXT("\t\t\t\t\"value\" : \"%s\""), *Attr.GetValue());
+				}
+				FileArchive->Logf(TEXT("\t\t\t}"));
+				bHasWrittenFirstAttr = true;
+			}
+
 			for (auto Attr : Attributes)
 			{
 				if (bHasWrittenFirstAttr)
@@ -550,3 +570,23 @@ void FAnalyticsProviderFileLogging::RecordCurrencyGiven(const FString& GameCurre
 		UE_LOG(LogFileLoggingAnalytics, Warning, TEXT("FAnalyticsProviderFileLogging::RecordCurrencyGiven called before StartSession. Ignoring."));
 	}
 }
+
+void FAnalyticsProviderFileLogging::SetDefaultEventAttributes(TArray<FAnalyticsEventAttribute>&& Attributes)
+{
+	DefaultEventAttributes = Attributes;
+}
+
+TArray<FAnalyticsEventAttribute> FAnalyticsProviderFileLogging::GetDefaultEventAttributesSafe() const
+{
+	return DefaultEventAttributes;
+}
+
+int32 FAnalyticsProviderFileLogging::GetDefaultEventAttributeCount() const
+{
+	return DefaultEventAttributes.Num();
+}
+
+FAnalyticsEventAttribute FAnalyticsProviderFileLogging::GetDefaultEventAttribute(int AttributeIndex) const
+{
+	return DefaultEventAttributes[AttributeIndex];
+}
\ No newline at end of file
diff --git a/Engine/Plugins/Runtime/Analytics/FileLogging/Source/FileLogging/Private/FileLoggingProvider.h b/Engine/Plugins/Runtime/Analytics/FileLogging/Source/FileLogging/Private/FileLoggingProvider.h
--- a/Engine/Plugins/Runtime/Analytics/FileLogging/Source/FileLogging/Private/FileLoggingProvider.h
+++ b/Engine/Plugins/Runtime/Analytics/FileLogging/Source/FileLogging/Private/FileLoggingProvider.h
@@ -64,4 +64,12 @@ public:
 	virtual void RecordCurrencyGiven(const FString& GameCurrencyType, int GameCurrencyAmount, const TArray<FAnalyticsEventAttribute>& EventAttrs) override;
 	virtual void RecordError(const FString& Error, const TArray<FAnalyticsEventAttribute>& EventAttrs) override;
 	virtual void RecordProgress(const FString& ProgressType, const FString& ProgressHierarchy, const TArray<FAnalyticsEventAttribute>& EventAttrs) override;
+
+	virtual void SetDefaultEventAttributes(TArray<FAnalyticsEventAttribute>&& Attributes)  override;
+	virtual TArray<FAnalyticsEventAttribute> GetDefaultEventAttributesSafe() const override;
+	virtual int32 GetDefaultEventAttributeCount() const  override;
+	virtual FAnalyticsEventAttribute GetDefaultEventAttribute(int AttributeIndex) const  override;
+
+private:
+	TArray<FAnalyticsEventAttribute> DefaultEventAttributes;
 };
diff --git a/Engine/Source/Runtime/Analytics/Analytics/Public/Interfaces/IAnalyticsProvider.h b/Engine/Source/Runtime/Analytics/Analytics/Public/Interfaces/IAnalyticsProvider.h
--- a/Engine/Source/Runtime/Analytics/Analytics/Public/Interfaces/IAnalyticsProvider.h
+++ b/Engine/Source/Runtime/Analytics/Analytics/Public/Interfaces/IAnalyticsProvider.h
@@ -122,6 +122,34 @@ public:
 	 */
 	virtual FString GetUserID() const = 0;
 
+
+	/**
+	 * Sets an array of attributes that will automatically be appended to any event that is sent.
+	 * Logical effect is like adding them to all events before calling RecordEvent.
+	 * Practically, it is implemented much more efficiently from a storage and allocation perspective.
+	 */
+	virtual void SetDefaultEventAttributes(TArray<FAnalyticsEventAttribute>&& Attributes) = 0;
+
+	/**
+	 * @return the current array of default attributes.
+	 */
+	virtual TArray<FAnalyticsEventAttribute> GetDefaultEventAttributesSafe() const = 0;
+
+	/**
+	 * Used with GetDefaultAttribute to iterate over the default attributes.
+	 *
+	 * @return the number of default attributes are currently being applied.
+	 */
+	virtual int32 GetDefaultEventAttributeCount() const = 0;
+
+	/**
+	 * Used with GetDefaultEventAttributeCount to iterate over the default attributes.
+	 *
+	 * Range checking is not done, similar to TArray. Use GetDefaultAttributeCount() first!
+	 * @return one attribute of the default attributes so we don't have to copy the entire attribute array.
+	 */
+	virtual FAnalyticsEventAttribute GetDefaultEventAttribute(int AttributeIndex) const = 0;
+
 	/**
 	 * Sets a user defined string as the build information/version for the session
 	 */
diff --git a/Engine/Source/Runtime/Analytics/AnalyticsET/Public/IAnalyticsProviderET.h b/Engine/Source/Runtime/Analytics/AnalyticsET/Public/IAnalyticsProviderET.h
--- a/Engine/Source/Runtime/Analytics/AnalyticsET/Public/IAnalyticsProviderET.h
+++ b/Engine/Source/Runtime/Analytics/AnalyticsET/Public/IAnalyticsProviderET.h
@@ -92,33 +92,6 @@ public:
 	 */
 	virtual void RecordEvent(FString&& EventName, const TArray<FAnalyticsEventAttribute>& Attributes) = 0;
 
-	/**
-	 * Sets an array of attributes that will automatically be appended to any event that is sent.
-	 * Logical effect is like adding them to all events before calling RecordEvent.
-	 * Practically, it is implemented much more efficiently from a storage and allocation perspective.
-	 */
-	virtual void SetDefaultEventAttributes(TArray<FAnalyticsEventAttribute>&& Attributes) = 0;
-
-	/** 
-	 * @return the current array of default attributes.
-	 */ 
-	virtual TArray<FAnalyticsEventAttribute> GetDefaultEventAttributesSafe() const = 0;
-
-	/**
-	 * Used with GetDefaultAttribute to iterate over the default attributes.
-	 * 
-	 * @return the number of default attributes are currently being applied.
-	 */
-	virtual int32 GetDefaultEventAttributeCount() const = 0;
-
-	/**
-	 * Used with GetDefaultEventAttributeCount to iterate over the default attributes.
-	 *
-	 * Range checking is not done, similar to TArray. Use GetDefaultAttributeCount() first!
-	 * @return one attribute of the default attributes so we don't have to copy the entire attribute array.
-	 */
-	virtual FAnalyticsEventAttribute GetDefaultEventAttribute(int AttributeIndex) const = 0;
-
 	/**
 	 * Updates the default URL endpoint and AltDomains.
 	 */
diff --git a/Engine/Source/Runtime/Analytics/AnalyticsSwrve/Private/AnalyticsSwrve.cpp b/Engine/Source/Runtime/Analytics/AnalyticsSwrve/Private/AnalyticsSwrve.cpp
--- a/Engine/Source/Runtime/Analytics/AnalyticsSwrve/Private/AnalyticsSwrve.cpp
+++ b/Engine/Source/Runtime/Analytics/AnalyticsSwrve/Private/AnalyticsSwrve.cpp
@@ -35,6 +35,11 @@ public:
 	virtual FString GetSessionID() const override;
 	virtual bool SetSessionID(const FString& InSessionID) override;
 
+	virtual void SetDefaultEventAttributes(TArray<FAnalyticsEventAttribute>&& Attributes)  override;
+	virtual TArray<FAnalyticsEventAttribute> GetDefaultEventAttributesSafe() const override;
+	virtual int32 GetDefaultEventAttributeCount() const  override;
+	virtual FAnalyticsEventAttribute GetDefaultEventAttribute(int AttributeIndex) const  override;
+
 	virtual void RecordEvent(const FString& EventName, const TArray<FAnalyticsEventAttribute>& Attributes) override;
 	virtual void RecordItemPurchase(const FString& ItemId, const FString& Currency, int PerItemCost, int ItemQuantity) override;
 	virtual void RecordCurrencyPurchase(const FString& GameCurrencyType, int GameCurrencyAmount, const FString& RealCurrencyType, float RealMoneyCost, const FString& PaymentProvider) override;
@@ -75,6 +80,8 @@ private:
 	static double kEventRateDetectionWindowTimeSec;
 	static int kEventRateDetectionCountThreshold;
 
+	TArray<FAnalyticsEventAttribute>		DefaultEventAttributes;
+
 	/**
 	 * Delegate called when an event Http request completes
 	 */
@@ -348,9 +355,22 @@ void FAnalyticsProviderSwrve::RecordEvent(const FString& EventName, const TArray
 {
 	// encode params as Json
 	FString EventParams = TEXT("swrve_payload=");
+
 	if (Attributes.Num() > 0)
 	{
 		EventParams += TEXT("{");
+
+		// Add the default attributes
+		for (int Ndx = 0; Ndx < DefaultEventAttributes.Num(); ++Ndx)
+		{
+			if (Ndx > 0)
+			{
+				EventParams += TEXT(",");
+			}
+			EventParams += FString(TEXT("\"")) + DefaultEventAttributes[Ndx].GetName() + TEXT("\": \"") + DefaultEventAttributes[Ndx].GetValue() + TEXT("\"");
+		}
+
+		// Add the event attributes
 		for (int Ndx=0;Ndx<Attributes.Num();++Ndx)
 		{
 			if (Ndx > 0)
@@ -424,4 +444,24 @@ void FAnalyticsProviderSwrve::ABTestRequestComplete(
 	}
 }
 
+void FAnalyticsProviderSwrve::SetDefaultEventAttributes(TArray<FAnalyticsEventAttribute>&& Attributes)
+{
+	DefaultEventAttributes = Attributes;
+}
+
+TArray<FAnalyticsEventAttribute> FAnalyticsProviderSwrve::GetDefaultEventAttributesSafe() const
+{
+	return DefaultEventAttributes;
+}
+
+int32 FAnalyticsProviderSwrve::GetDefaultEventAttributeCount() const
+{
+	return DefaultEventAttributes.Num();
+}
+
+FAnalyticsEventAttribute FAnalyticsProviderSwrve::GetDefaultEventAttribute(int AttributeIndex) const
+{
+	return DefaultEventAttributes[AttributeIndex];
+}
+
 #endif // PLATFORM_DESKTOP
