Fix for IOS and AndroidAdjust & AppBoyClient AnalyticsProvider for Android and IOS
Moved default telemetry logs to Saved/Telemetry

diff --git a/Engine/Plugins/Experimental/StudioTelemetry/Source/AnalyticsLog/Private/AnalyticsProviderLog.cpp b/Engine/Plugins/Experimental/StudioTelemetry/Source/AnalyticsLog/Private/AnalyticsProviderLog.cpp
--- a/Engine/Plugins/Experimental/StudioTelemetry/Source/AnalyticsLog/Private/AnalyticsProviderLog.cpp
+++ b/Engine/Plugins/Experimental/StudioTelemetry/Source/AnalyticsLog/Private/AnalyticsProviderLog.cpp
@@ -20,7 +20,7 @@ FAnalyticsProviderLog::FAnalyticsProviderLog(const FAnalyticsProviderConfigurati
 	if (PathName.IsEmpty())
 	{
 		// Use default output path
-		PathName = FPaths::EngineSavedDir();
+		PathName = FPaths::ProjectSavedDir() / TEXT("Telemetry");
 	}
 
 	// Create the full output path
@@ -91,32 +91,30 @@ void FAnalyticsProviderLog::EndSession()
 
 	if (FileWriter)
 	{
+		FileWriter->Flush();
 		FileWriter->Close();
 	}
 }
 
 void FAnalyticsProviderLog::RecordEvent(const FString& EventName, const TArray<FAnalyticsEventAttribute>& Attributes)
 {	
-	TStringBuilder<1024> StringBuilder;
-
-	StringBuilder.Appendf(TEXT("{{EventName=%s}"), *EventName);
-
-	// Log the default attributes
-	for (const FAnalyticsEventAttribute& Attribute : DefaultEventAttributes)
-	{
-		StringBuilder.Appendf(TEXT(",{%s=%s}"), *Attribute.GetName(), *Attribute.GetValue());
-	}
-
-	// Log the event attributes
-	for ( const FAnalyticsEventAttribute& Attribute : Attributes )
-	{
-		StringBuilder.Appendf(TEXT("{%s=%s}"), *Attribute.GetName(), *Attribute.GetValue());
-	}
-
-	StringBuilder.Appendf(TEXT("}\n"));
-
 	if (FileWriter)
 	{
-		FileWriter->Serialize(TCHAR_TO_ANSI(StringBuilder.ToString()), StringBuilder.Len());
+		FileWriter->Logf(TEXT("{{EventName=%s}"), *EventName);
+
+		// Log the default attributes
+		for (const FAnalyticsEventAttribute& Attribute : DefaultEventAttributes)
+		{
+			FileWriter->Logf(TEXT(",{%s=%s}"), *Attribute.GetName(), *Attribute.GetValue());
+		}
+
+		// Log the event attributes
+		for ( const FAnalyticsEventAttribute& Attribute : Attributes )
+		{
+			FileWriter->Logf(TEXT(",{%s=%s}"), *Attribute.GetName(), *Attribute.GetValue());
+		}
+
+		FileWriter->Logf(TEXT("}\n"));
+		FileWriter->Flush();
 	}
 }
\ No newline at end of file
diff --git a/Engine/Plugins/Runtime/Analytics/Adjust/Source/AndroidAdjust/Private/AndroidAdjust.cpp b/Engine/Plugins/Runtime/Analytics/Adjust/Source/AndroidAdjust/Private/AndroidAdjust.cpp
--- a/Engine/Plugins/Runtime/Analytics/Adjust/Source/AndroidAdjust/Private/AndroidAdjust.cpp
+++ b/Engine/Plugins/Runtime/Analytics/Adjust/Source/AndroidAdjust/Private/AndroidAdjust.cpp
@@ -321,6 +321,11 @@ void FAnalyticsProviderAdjust::RecordEvent(const FString& EventName, const TArra
 {
 #if WITH_ADJUST
 	FString* EventTokenRef = EventMap.Find(EventName);
+
+	TArray<FAnalyticsEventAttribute> EventAttributes;
+	EventAttributes.Append(Attributes);
+	EventAttributes.Append(DefaultEventAttributes);
+
 	if (EventTokenRef != nullptr)
 	{
 		FString EventToken = *EventTokenRef;
@@ -328,11 +333,11 @@ void FAnalyticsProviderAdjust::RecordEvent(const FString& EventName, const TArra
 		AndroidThunkCpp_Adjust_Event_ResetCallbackParameters();
 		AndroidThunkCpp_Adjust_Event_ResetPartnerParameters();
 
-		const int32 AttrCount = Attributes.Num();
+		const int32 AttrCount = EventAttributes.Num();
 		if (AttrCount > 0)
 		{
 			// add event attributes
-			for (auto Attr : Attributes)
+			for (auto Attr : EventAttributes)
 			{
 				AndroidThunkCpp_Adjust_Event_AddCallbackParameter(Attr.GetName(), Attr.GetValue());
 			}
@@ -483,3 +488,23 @@ void FAnalyticsProviderAdjust::RecordProgress(const FString& ProgressType, const
 	UE_LOG(LogAnalytics, Warning, TEXT("WITH_ADJUST=0. Are you missing the SDK?"));
 #endif
 }
+
+void FAnalyticsProviderAdjust::SetDefaultEventAttributes(TArray<FAnalyticsEventAttribute>&& Attributes)
+{
+	DefaultEventAttributes = Attributes;
+}
+
+TArray<FAnalyticsEventAttribute> FAnalyticsProviderAdjust::GetDefaultEventAttributesSafe() const
+{
+	return DefaultEventAttributes;
+}
+
+int32 FAnalyticsProviderAdjust::GetDefaultEventAttributeCount() const
+{
+	return DefaultEventAttributes.Num();
+}
+
+FAnalyticsEventAttribute FAnalyticsProviderAdjust::GetDefaultEventAttribute(int AttributeIndex) const
+{
+	return DefaultEventAttributes[AttributeIndex];
+}
\ No newline at end of file
diff --git a/Engine/Plugins/Runtime/Analytics/Adjust/Source/AndroidAdjust/Private/AndroidAdjustProvider.h b/Engine/Plugins/Runtime/Analytics/Adjust/Source/AndroidAdjust/Private/AndroidAdjustProvider.h
--- a/Engine/Plugins/Runtime/Analytics/Adjust/Source/AndroidAdjust/Private/AndroidAdjustProvider.h
+++ b/Engine/Plugins/Runtime/Analytics/Adjust/Source/AndroidAdjust/Private/AndroidAdjustProvider.h
@@ -53,4 +53,12 @@ public:
 
 	virtual void RecordError(const FString& Error, const TArray<FAnalyticsEventAttribute>& EventAttrs) override;
 	virtual void RecordProgress(const FString& ProgressType, const FString& ProgressHierarchy, const TArray<FAnalyticsEventAttribute>& EventAttrs) override;
+
+	virtual void SetDefaultEventAttributes(TArray<FAnalyticsEventAttribute>&& Attributes)  override;
+	virtual TArray<FAnalyticsEventAttribute> GetDefaultEventAttributesSafe() const override;
+	virtual int32 GetDefaultEventAttributeCount() const  override;
+	virtual FAnalyticsEventAttribute GetDefaultEventAttribute(int AttributeIndex) const  override;
+
+private:
+	TArray<FAnalyticsEventAttribute>		DefaultEventAttributes;
 };
diff --git a/Engine/Plugins/Runtime/Analytics/Adjust/Source/IOSAdjust/Private/IOSAdjust.cpp b/Engine/Plugins/Runtime/Analytics/Adjust/Source/IOSAdjust/Private/IOSAdjust.cpp
--- a/Engine/Plugins/Runtime/Analytics/Adjust/Source/IOSAdjust/Private/IOSAdjust.cpp
+++ b/Engine/Plugins/Runtime/Analytics/Adjust/Source/IOSAdjust/Private/IOSAdjust.cpp
@@ -265,15 +265,19 @@ void FAnalyticsProviderAdjust::RecordEvent(const FString& EventName, const TArra
 	FString* EventTokenRef = EventMap.Find(EventName);
 	if (EventTokenRef != nullptr)
 	{
+		TArray<FAnalyticsEventAttribute> EventAttributes;
+		EventAttributes.Append(Attributes);
+		EventAttributes.Append(DefaultEventAttributes);
+
 		FString EventToken = *EventTokenRef;
 		NSString* IOSEventToken = [NSString stringWithFString : EventToken];
 		
 		ADJEvent *event = [ADJEvent eventWithEventToken:IOSEventToken];
-		const int32 AttrCount = Attributes.Num();
+		const int32 AttrCount = EventAttributes.Num();
 		if (AttrCount > 0)
 		{
 			// add event attributes
-			for (auto Attr : Attributes)
+			for (auto Attr : EventAttributes)
 			{
 				NSString* IOSKey = [NSString stringWithFString : Attr.GetName()];
 				NSString* IOSValue = [NSString stringWithFString : Attr.GetValue()];
@@ -430,3 +434,24 @@ void FAnalyticsProviderAdjust::RecordProgress(const FString& ProgressType, const
 	UE_LOG(LogAnalytics, Warning, TEXT("WITH_ADJUST=0. Are you missing the SDK?"));
 #endif
 }
+
+
+void FAnalyticsProviderAdjust::SetDefaultEventAttributes(TArray<FAnalyticsEventAttribute>&& Attributes)
+{
+	DefaultEventAttributes = Attributes;
+}
+
+TArray<FAnalyticsEventAttribute> FAnalyticsProviderAdjust::GetDefaultEventAttributesSafe() const
+{
+	return DefaultEventAttributes;
+}
+
+int32 FAnalyticsProviderAdjust::GetDefaultEventAttributeCount() const
+{
+	return DefaultEventAttributes.Num();
+}
+
+FAnalyticsEventAttribute FAnalyticsProviderAdjust::GetDefaultEventAttribute(int AttributeIndex) const
+{
+	return DefaultEventAttributes[AttributeIndex];
+}
\ No newline at end of file
diff --git a/Engine/Plugins/Runtime/Analytics/Adjust/Source/IOSAdjust/Private/IOSAdjustProvider.h b/Engine/Plugins/Runtime/Analytics/Adjust/Source/IOSAdjust/Private/IOSAdjustProvider.h
--- a/Engine/Plugins/Runtime/Analytics/Adjust/Source/IOSAdjust/Private/IOSAdjustProvider.h
+++ b/Engine/Plugins/Runtime/Analytics/Adjust/Source/IOSAdjust/Private/IOSAdjustProvider.h
@@ -53,4 +53,12 @@ public:
 
 	virtual void RecordError(const FString& Error, const TArray<FAnalyticsEventAttribute>& EventAttrs) override;
 	virtual void RecordProgress(const FString& ProgressType, const FString& ProgressHierarchy, const TArray<FAnalyticsEventAttribute>& EventAttrs) override;
+
+	virtual void SetDefaultEventAttributes(TArray<FAnalyticsEventAttribute>&& Attributes)  override;
+	virtual TArray<FAnalyticsEventAttribute> GetDefaultEventAttributesSafe() const override;
+	virtual int32 GetDefaultEventAttributeCount() const  override;
+	virtual FAnalyticsEventAttribute GetDefaultEventAttribute(int AttributeIndex) const  override;
+
+private:
+	TArray<FAnalyticsEventAttribute>		DefaultEventAttributes;
 };
