Changed AnalyticsLog format to a standard JSON Newline stream

diff --git a/Engine/Plugins/Experimental/StudioTelemetry/Source/AnalyticsLog/Private/AnalyticsProviderLog.cpp b/Engine/Plugins/Experimental/StudioTelemetry/Source/AnalyticsLog/Private/AnalyticsProviderLog.cpp
--- a/Engine/Plugins/Experimental/StudioTelemetry/Source/AnalyticsLog/Private/AnalyticsProviderLog.cpp
+++ b/Engine/Plugins/Experimental/StudioTelemetry/Source/AnalyticsLog/Private/AnalyticsProviderLog.cpp
@@ -12,7 +12,7 @@ FAnalyticsProviderLog::FAnalyticsProviderLog(const FAnalyticsProviderConfigurati
 	if (FileName.IsEmpty())
 	{
 		// Use default filename
-		FileName = TEXT("Telemetry.log");
+		FileName = TEXT("Telemetry.json");
 	}
 
 	FString PathName = GetConfigValue.Execute(TEXT("PathName"), true);
@@ -100,21 +100,27 @@ void FAnalyticsProviderLog::RecordEvent(const FString& EventName, const TArray<F
 {	
 	if (FileWriter)
 	{
-		FileWriter->Logf(TEXT("{{EventName=%s}"), *EventName);
+		TStringBuilder<1024> Builder;
+
+		// Log event as Newline - delimited JSON
+		Builder.Appendf(TEXT("{\"EventName\":\"%s\""), *EventName);
+
+		// Add the event timestamp field
+		Builder.Appendf(TEXT(",\"TimestampUTC\":\"%f\""), FDateTime::UtcNow().ToUnixTimestampDecimal());
 
 		// Log the default attributes
 		for (const FAnalyticsEventAttribute& Attribute : DefaultEventAttributes)
 		{
-			FileWriter->Logf(TEXT(",{%s=%s}"), *Attribute.GetName(), *Attribute.GetValue());
+			Builder.Appendf(TEXT(",\"%s\":\"%s\""), *Attribute.GetName(), *Attribute.GetValue());
 		}
 
 		// Log the event attributes
 		for ( const FAnalyticsEventAttribute& Attribute : Attributes )
 		{
-			FileWriter->Logf(TEXT(",{%s=%s}"), *Attribute.GetName(), *Attribute.GetValue());
+			Builder.Appendf(TEXT(",\"%s\":\"%s\""), *Attribute.GetName(), *Attribute.GetValue());
 		}
 
-		FileWriter->Logf(TEXT("}\n"));
+		FileWriter->Logf(TEXT("%s}"),Builder.ToString());
 		FileWriter->Flush();
 	}
 }
\ No newline at end of file
