Studio Telemetry Update:
Added IAnalyticsTracer and IAnalyticsSpan interface that superseeds the FFlowTracker
StudioTelemetry no longer uses FFLowTracker and uses the simpler IAnalyticsTracer API instead inspired by OpenTelemetry Tracer/Span model.
Deprecated FFLowTracker as of 5.4
Added lots of documentation to the StudioTelemetry plugin code.
Moved creation of the AnalyticsProviderLog to the public BaseEngine.ini so it can be used as a public facing example. This will simply write out a JSON file to the Saved/Telemetry folder and overwrite it per session
Added new callbacks for PIE so that we can decopule when a PIE work has been created and is ready from when ALL PIE worlds have been created and are ready which makes TimeToPIE accurate for all projects.
Rolled in some shutdown recommendations from Paul.Chipchase.
Various fixes and improvements to the JSON file logigng example. Now works correctly with Tableau import.

diff --git a/Engine/Config/BaseEngine.ini b/Engine/Config/BaseEngine.ini
--- a/Engine/Config/BaseEngine.ini
+++ b/Engine/Config/BaseEngine.ini
@@ -3530,3 +3530,12 @@ IntervalSeconds=1
 +Listeners=(MetricName=PingBucketInt7, Class=/Script/Engine.NetworkMetricsPerfCounters)
 +Listeners=(MetricName=ClosedConnectionsDueToReliableBufferOverflow, Class=/Script/Engine.NetworkMetricsPerfCounters)
 +Listeners=(MetricName=AddedConnections, Class=/Script/Engine.NetworkMetricsPerfCounters)
+
+; Studio Telemetry Settings
+[StudioTelemetry.Log] 
+Name=LogAnalytics
+ProviderModule=AnalyticsLog
+UsageType=Editor
+; Override the default output location from Saved/Telemetry/Telemetry.json
+; FileName="somefile.json"
+; FolderPath="somefolder"
\ No newline at end of file
diff --git a/Engine/Plugins/Experimental/StudioTelemetry/Source/AnalyticsLog/AnalyticsLog.Build.cs b/Engine/Plugins/Experimental/StudioTelemetry/Source/AnalyticsLog/AnalyticsLog.Build.cs
--- a/Engine/Plugins/Experimental/StudioTelemetry/Source/AnalyticsLog/AnalyticsLog.Build.cs
+++ b/Engine/Plugins/Experimental/StudioTelemetry/Source/AnalyticsLog/AnalyticsLog.Build.cs
@@ -15,8 +15,7 @@ namespace UnrealBuildTool.Rules
 			PrivateDependencyModuleNames.AddRange(
 				new string[] {
 				"CoreUObject",
-				"Analytics",
-				"AnalyticsET"
+				"Analytics"
 				}
 			);
 		}
diff --git a/Engine/Plugins/Experimental/StudioTelemetry/Source/AnalyticsLog/Private/AnalyticsProviderLog.cpp b/Engine/Plugins/Experimental/StudioTelemetry/Source/AnalyticsLog/Private/AnalyticsProviderLog.cpp
--- a/Engine/Plugins/Experimental/StudioTelemetry/Source/AnalyticsLog/Private/AnalyticsProviderLog.cpp
+++ b/Engine/Plugins/Experimental/StudioTelemetry/Source/AnalyticsLog/Private/AnalyticsProviderLog.cpp
@@ -15,16 +15,16 @@ FAnalyticsProviderLog::FAnalyticsProviderLog(const FAnalyticsProviderConfigurati
 		FileName = TEXT("Telemetry.json");
 	}
 
-	FString PathName = GetConfigValue.Execute(TEXT("PathName"), true);
+	FString FolderPath = GetConfigValue.Execute(TEXT("FolderPath"), true);
 
-	if (PathName.IsEmpty())
+	if (FolderPath.IsEmpty())
 	{
 		// Use default output path
-		PathName = FPaths::ProjectSavedDir() / TEXT("Telemetry");
+		FolderPath = FPaths::ProjectSavedDir() / TEXT("Telemetry");
 	}
 
 	// Create the full output path
-	FString FilePath = PathName / FileName;
+	FString FilePath = FolderPath / FileName;
 	FileWriter = TUniquePtr<FArchive>(IFileManager::Get().CreateFileWriter(*FilePath, FILEWRITE_EvenIfReadOnly));
 }
 
@@ -98,6 +98,8 @@ void FAnalyticsProviderLog::EndSession()
 
 void FAnalyticsProviderLog::RecordEvent(const FString& EventName, const TArray<FAnalyticsEventAttribute>& Attributes)
 {	
+	static uint32 RecordId(0);
+
 	if (FileWriter)
 	{
 		TStringBuilder<1024> Builder;
@@ -106,18 +108,27 @@ void FAnalyticsProviderLog::RecordEvent(const FString& EventName, const TArray<F
 		Builder.Appendf(TEXT("{\"EventName\":\"%s\""), *EventName);
 
 		// Add the event timestamp field
-		Builder.Appendf(TEXT(",\"TimestampUTC\":\"%f\""), FDateTime::UtcNow().ToUnixTimestampDecimal());
+		Builder.Appendf(TEXT(",\"TimestampUTC\":%f"), FDateTime::UtcNow().ToUnixTimestampDecimal());
 
-		// Log the default attributes
-		for (const FAnalyticsEventAttribute& Attribute : DefaultEventAttributes)
-		{
-			Builder.Appendf(TEXT(",\"%s\":\"%s\""), *Attribute.GetName(), *Attribute.GetValue());
-		}
+		// Add the record Id and increment it
+		Builder.Appendf(TEXT(",\"RecordId\":%d"), RecordId++);
+
+		// Accumulate all the attributes together. We could have had two loops but this seems cleaner
+		TArray<FAnalyticsEventAttribute> EventAttributes(DefaultEventAttributes);
+		EventAttributes.Append(Attributes);
 
-		// Log the event attributes
-		for ( const FAnalyticsEventAttribute& Attribute : Attributes )
+		// Add all the attributes
+		for (const FAnalyticsEventAttribute& Attribute : EventAttributes)
 		{
-			Builder.Appendf(TEXT(",\"%s\":\"%s\""), *Attribute.GetName(), *Attribute.GetValue());
+			// This should be almost nearly true, but we should check and JSON'ify as needed
+			if (Attribute.IsJsonFragment())
+			{
+				Builder.Appendf(TEXT(",\"%s\":%s"), *Attribute.GetName(), *Attribute.GetValue());
+			}
+			else
+			{
+				Builder.Appendf(TEXT(",\"%s\":\"%s\""), *Attribute.GetName(), *Attribute.GetValue());
+			}
 		}
 
 		FileWriter->Logf(TEXT("%s}"),Builder.ToString());
diff --git a/Engine/Plugins/Experimental/StudioTelemetry/Source/AnalyticsLog/Private/AnalyticsProviderLog.h b/Engine/Plugins/Experimental/StudioTelemetry/Source/AnalyticsLog/Private/AnalyticsProviderLog.h
--- a/Engine/Plugins/Experimental/StudioTelemetry/Source/AnalyticsLog/Private/AnalyticsProviderLog.h
+++ b/Engine/Plugins/Experimental/StudioTelemetry/Source/AnalyticsLog/Private/AnalyticsProviderLog.h
@@ -2,13 +2,14 @@
 
 #pragma once
 
-#if WITH_EDITOR
-
 #include "AnalyticsProviderConfigurationDelegate.h"
 #include "Interfaces/IAnalyticsProvider.h"
 
 /**
- * Implementation of the IAnalyticsProviderET interface that exports to OpenTelemetry
+ * Implementation of the IAnalyticsProviderET interface that exports telemetry events to a file as Newline - delimited JSON
+ * By default,the log file is written to Saved/Telemetry folder of the application.
+ * FileName and FolderPath can be overridden in the configuration
+ * Here as a simple example of how a developer might implement and configure their own analytics provider for use with the StudioTelemerty plugin
  */
 class FAnalyticsProviderLog : public IAnalyticsProvider
 {
@@ -41,5 +42,3 @@ private:
 	TArray<FAnalyticsEventAttribute>		DefaultEventAttributes;	
 	TUniquePtr<FArchive>					FileWriter;
 };
-
-#endif
\ No newline at end of file
diff --git a/Engine/Plugins/Experimental/StudioTelemetry/Source/AnalyticsLog/Public/AnalyticsLog.h b/Engine/Plugins/Experimental/StudioTelemetry/Source/AnalyticsLog/Public/AnalyticsLog.h
--- a/Engine/Plugins/Experimental/StudioTelemetry/Source/AnalyticsLog/Public/AnalyticsLog.h
+++ b/Engine/Plugins/Experimental/StudioTelemetry/Source/AnalyticsLog/Public/AnalyticsLog.h
@@ -5,7 +5,6 @@
 #if WITH_EDITOR
 
 #include "AnalyticsProviderConfigurationDelegate.h"
-#include "AnalyticsET.h"
 #include "Containers/Array.h"
 #include "Containers/UnrealString.h"
 #include "CoreMinimal.h"
@@ -20,7 +19,7 @@ class IAnalyticsProvider;
 class IAnalyticsProviderET;
 
 /**
- *  Public implementation of EpicGames.MCP.AnalyticsProvider
+ *  Public implementation of IAnalyticsProviderModule that returns a JSON formatted FAnalayticsProviderLog object
  */
 class FAnalyticsLog : public IAnalyticsProviderModule
 {
@@ -30,8 +29,7 @@ class FAnalyticsLog : public IAnalyticsProviderModule
 public:
 	/**
 	 * Singleton-like access to this module's interface.  This is just for convenience!
-	 * Beware of calling this during the shutdown phase, though.  Your module might have been unloaded already.
-	 *
+	 * Beware of calling this during the shutdown phase, though.  Your module might have been unloaded already
 	 * @return Returns singleton instance, loading the module on demand if needed
 	 */
 	static inline FAnalyticsLog& Get()
@@ -39,19 +37,8 @@ public:
 		return FModuleManager::LoadModuleChecked< FAnalyticsLog >( "AnalyticsLog" );
 	}
 
-	//--------------------------------------------------------------------------
-	// Configuration functionality
-	//--------------------------------------------------------------------------
-public:
-	
-	//--------------------------------------------------------------------------
-	// provider factory functions
-	//--------------------------------------------------------------------------
-public:
 	/**
 	 * IAnalyticsProviderModule interface.
-	 * Creates the analytics provider given a configuration delegate.
-	 * The keys required exactly match the field names in the Config object. 
 	 */
 	ANALYTICSLOG_API virtual TSharedPtr<IAnalyticsProvider> CreateAnalyticsProvider(const FAnalyticsProviderConfigurationDelegate& GetConfigValue) const override;
 
diff --git a/Engine/Plugins/Experimental/StudioTelemetry/Source/StudioTelemetry/Private/AnalyticsProviderMulticast.h b/Engine/Plugins/Experimental/StudioTelemetry/Source/StudioTelemetry/Private/AnalyticsProviderMulticast.h
--- a/Engine/Plugins/Experimental/StudioTelemetry/Source/StudioTelemetry/Private/AnalyticsProviderMulticast.h
+++ b/Engine/Plugins/Experimental/StudioTelemetry/Source/StudioTelemetry/Private/AnalyticsProviderMulticast.h
@@ -6,7 +6,20 @@
 #include "Interfaces/IAnalyticsProvider.h"
 
 /**
- * Implementation of the IAnalyticsProvider interface that forwards the API calls to an array of IAnalyticsProviderET interfaces
+ * Implementation of the IAnalyticsProvider interface that forwards the API calls to an array of IAnalyticsProvider interfaces
+ * 
+ * IAnalyticsProviders are specified in the Engine.ini interface and are constructed and configured automatically.
+ *
+ * AnalyticsLog for example in Engine/Config/BaseEngine.ini
+ * 
+ * ; Studio Telemetry Settings
+ * [StudioTelemetry.Log]
+ * Name=LogAnalytics
+ * ProviderModule=AnalyticsLog
+ * UsageType=Editor
+ * 
+ * This instructs the Multicast provider to create a Provider from the AnalyticsLog module via the IANalyticsProviderModule interface. 
+ * See FAnalayticsLog and FAnalyticsProviderLog for more details.
  */
 class FAnalyticsProviderMulticast : public IAnalyticsProvider
 {
diff --git a/Engine/Plugins/Experimental/StudioTelemetry/Source/StudioTelemetry/Private/StudioTelemetry.cpp b/Engine/Plugins/Experimental/StudioTelemetry/Source/StudioTelemetry/Private/StudioTelemetry.cpp
--- a/Engine/Plugins/Experimental/StudioTelemetry/Source/StudioTelemetry/Private/StudioTelemetry.cpp
+++ b/Engine/Plugins/Experimental/StudioTelemetry/Source/StudioTelemetry/Private/StudioTelemetry.cpp
@@ -10,6 +10,7 @@
 
 #include "Analytics.h"
 #include "AnalyticsProviderMulticast.h"
+#include "AnalyticsTracer.h"
 #include "BuildSettings.h"
 #include "RHI.h"
 
@@ -73,19 +74,20 @@ void FStudioTelemetry::ShutdownModule()
 
 void FStudioTelemetry::EndSession()
 {
-	if (AnalyticsFlowTracker.IsValid())
+	// End session for the tracer and the provider
+	if (AnalyticsTracer.IsValid())
 	{
-		AnalyticsFlowTracker->EndSession();
-		AnalyticsFlowTracker.Reset();
+		AnalyticsTracer->EndSession();
+		AnalyticsTracer.Reset();
 	}
 
 	if (AnalyticsProvider.IsValid())
 	{
 		AnalyticsProvider->EndSession();
 		AnalyticsProvider.Reset();
-	}
 
-	UE_LOG(LogStudioTelemetry, Log, TEXT("Ended StudioTelemetry Session"));
+		UE_LOG(LogStudioTelemetry, Log, TEXT("Ended StudioTelemetry Session"));
+	}
 }
 
 void FStudioTelemetry::StartSession()
@@ -172,10 +174,13 @@ void FStudioTelemetry::StartSession()
 		// Start the analytics session
 		AnalyticsProvider->StartSession();
 
-		// Make the flow tracker and start the session
-		AnalyticsFlowTracker = MakeShared<FAnalyticsFlowTracker>();
-		AnalyticsFlowTracker->SetProvider(AnalyticsProvider);
-		AnalyticsFlowTracker->StartSession();
+		// Create the IAnalyticsTracer interface
+		AnalyticsTracer = FAnalytics::Get().CreateAnalyticsTracer();
+		AnalyticsTracer->SetProvider(AnalyticsProvider);
+		AnalyticsTracer->StartSession();
+
+		// Bind the pre-exit callback
+		FCoreDelegates::OnEnginePreExit.AddRaw(&FStudioTelemetry::Get(), &FStudioTelemetry::EndSession);
 
 		UE_LOG(LogStudioTelemetry, Log, TEXT("Started StudioTelemetry Session"));
 	}
@@ -211,7 +216,32 @@ TWeakPtr<IAnalyticsProvider> FStudioTelemetry::GetProvider(const FString& Name)
 	return AnalyticsProvider.IsValid()? AnalyticsProvider->GetAnalyticsProvider(Name) : TWeakPtr<IAnalyticsProvider>();
 }
 
-TWeakPtr<FAnalyticsFlowTracker> FStudioTelemetry::GetFlowTracker()
+TWeakPtr<IAnalyticsTracer> FStudioTelemetry::GetTracer()
+{
+	return AnalyticsTracer;
+}
+
+TSharedPtr<IAnalyticsSpan> FStudioTelemetry::StartSpan(const FName Name, const TArray<FAnalyticsEventAttribute>& AdditionalAttributes)
 {
-	return AnalyticsFlowTracker;
-}
\ No newline at end of file
+	return AnalyticsTracer.IsValid() ? AnalyticsTracer->StartSpan(Name, AdditionalAttributes) : TSharedPtr<IAnalyticsSpan>();
+}
+
+TSharedPtr<IAnalyticsSpan> FStudioTelemetry::StartSpan(const FName Name, TSharedPtr<IAnalyticsSpan> ParentSpan, const TArray<FAnalyticsEventAttribute>& AdditionalAttributes)
+{
+	return AnalyticsTracer.IsValid() ? AnalyticsTracer->StartSpan(Name, ParentSpan, AdditionalAttributes)  : TSharedPtr<IAnalyticsSpan>();
+}
+
+bool FStudioTelemetry::EndSpan(TSharedPtr<IAnalyticsSpan> Span, const TArray<FAnalyticsEventAttribute>& AdditionalAttributes)
+{
+	return AnalyticsTracer.IsValid() ? AnalyticsTracer->EndSpan(Span, AdditionalAttributes) : false;
+}
+
+bool FStudioTelemetry::EndSpan(const FName& Name, const TArray<FAnalyticsEventAttribute>& AdditionalAttributes)
+{
+	return AnalyticsTracer.IsValid() ? AnalyticsTracer->EndSpan(Name, AdditionalAttributes) : false;
+}
+
+TSharedPtr<IAnalyticsSpan> FStudioTelemetry::GetSpan(const FName Name)
+{
+	return AnalyticsTracer.IsValid() ? AnalyticsTracer->GetSpan(Name) : TSharedPtr<IAnalyticsSpan>();
+}
diff --git a/Engine/Plugins/Experimental/StudioTelemetry/Source/StudioTelemetry/Private/StudioTelemetryEditor.cpp b/Engine/Plugins/Experimental/StudioTelemetry/Source/StudioTelemetry/Private/StudioTelemetryEditor.cpp
--- a/Engine/Plugins/Experimental/StudioTelemetry/Source/StudioTelemetry/Private/StudioTelemetryEditor.cpp
+++ b/Engine/Plugins/Experimental/StudioTelemetry/Source/StudioTelemetry/Private/StudioTelemetryEditor.cpp
@@ -5,7 +5,7 @@
 #if WITH_EDITOR
 
 #include "StudioTelemetry.h"
-#include "AnalyticsFlowTracker.h"
+#include "AnalyticsTracer.h"
 #include "CollectionManagerModule.h"
 #include "ContentBrowserModule.h"
 #include "TelemetryRouter.h"
@@ -398,187 +398,60 @@ void FStudioTelemetryEditor::Initialize()
 {
 	SessionStartTime = FPlatformTime::Seconds();
 
-	TSharedPtr<FAnalyticsFlowTracker> FlowTracker = FStudioTelemetry::Get().GetFlowTracker().Pin();
+	// Install Editor Mode callbacks
+	// 
+	// Start Editor and Editor Boot span. Note : this will only start when the plugin is loaded and as such will miss any activity that runs beforehand
+	EditorSpan = FStudioTelemetry::Get().StartSpan(EditorSpanName);
+	EditorBootSpan = FStudioTelemetry::Get().StartSpan(EditorBootSpanName);
 
-	if (FlowTracker.IsValid())
-	{
-		// Start Editor Boot Flow immediately
-		EditorBootFlowGuid = FlowTracker->StartFlow(TEXT("Editor Boot"));
-		FlowTracker->StartSubFlow(TEXT("Editor Boot"), EditorBootFlowGuid);
-	}
-
-	FEditorDelegates::BeginPIE.AddLambda([this](bool)
-		{
-			TSharedPtr<FAnalyticsFlowTracker> FlowTracker = FStudioTelemetry::Get().GetFlowTracker().Pin();
+	TArray<FAnalyticsEventAttribute> Attributes;
+	Attributes.Emplace(FAnalyticsEventAttribute(TEXT("LevelName"), TEXT("EditorBoot")));
+	Attributes.Emplace(FAnalyticsEventAttribute(TEXT("MapName"), TEXT("EditorBoot")));
 
-			if (FlowTracker.IsValid())
-			{
-				// Finish the Editor Interaction Flow
-				FlowTracker->EndFlow(InteractiveEditorFlowGuid);
-
-				// Start PIE Flow
-				PIEStartTime = FPlatformTime::Seconds();
-				PIEFlowGuid = FlowTracker->StartFlow(TEXT("Play In Editor"));
-				PIEInitializeSubFlowGuid = FlowTracker->StartSubFlow(TEXT("PIE Initialize"), PIEFlowGuid);
-			}			
-		});
+	EditorBootSpan->AddAttributes(Attributes);
 
-	FEditorDelegates::PostPIEStarted.AddLambda([this](bool)
+	FEditorDelegates::OnMapLoad.AddLambda([this](const FString& MapName, FCanLoadMap& OutCanLoadMap)
 		{
-			TSharedPtr<FAnalyticsFlowTracker> FlowTracker = FStudioTelemetry::Get().GetFlowTracker().Pin();
-
-			if (FlowTracker.IsValid())
-			{
-				// Called when PIE has finally started
-				FlowTracker->EndSubFlow(PIEInitializeSubFlowGuid);
-			}
+			// The Editor loads a new map
+			EditorLoadMapSpan = FStudioTelemetry::Get().StartSpan(EditorLoadMapSpanName);
 		});
 
-	FEditorDelegates::EndPIE.AddLambda([this](bool)
+	FEditorDelegates::OnMapOpened.AddLambda([this](const FString& MapName, bool Unused)
 		{
-			TSharedPtr<FAnalyticsFlowTracker> FlowTracker = FStudioTelemetry::Get().GetFlowTracker().Pin();
-
-			if (FlowTracker.IsValid())
-			{
-				// End PIE Flow and any SubFlows within it
-				FlowTracker->EndFlow(PIEFlowGuid, true);
-
-				// Restart the Interactive Editor Flow
-				FlowTracker->StartFlow(TEXT("Interactive Editor"));
-			}
-		});
-
-	FWorldDelegates::OnPIEStarted.AddLambda([this](UGameInstance* GameInstance)
-		{
-			TSharedPtr<FAnalyticsFlowTracker> FlowTracker = FStudioTelemetry::Get().GetFlowTracker().Pin();
-
-			if (FlowTracker.IsValid())
-			{
-				PIELoadMapStartTime = FPlatformTime::Seconds();
-				PIELoadMapSubFlowGuid = FlowTracker->StartSubFlow(TEXT("PIE Load Map"), PIEFlowGuid);
-			}
-		});
-
-	FWorldDelegates::OnPIEReady.AddLambda([this](UGameInstance* GameInstance)
-		{	
-			const FString MapName = FPaths::GetBaseFilename(GameInstance->PIEMapName);
+			// The new editor map was actually opened
+			LevelName = FPaths::GetBaseFilename(MapName);
 
 			TArray<FAnalyticsEventAttribute> Attributes;
 			Attributes.Emplace(FAnalyticsEventAttribute(TEXT("LevelName"), LevelName));
-			Attributes.Emplace(FAnalyticsEventAttribute(TEXT("MapName"), MapName));
-			
-			// Currently there is no obvious way to know when PIE has finished loading everything but loading map FrontEnd is good enough for Epic products.
-			if (MapName.Contains(TEXT("Frontend")))
-			{
-				static bool IsFirstTimeToPIE = true;
-
-				PIEStartupTime = FPlatformTime::Seconds() - PIEStartTime;
+			Attributes.Emplace(FAnalyticsEventAttribute(TEXT("MapName"), LevelName));
 
-				if (IsFirstTimeToPIE == true)
-				{
-					// Record the absolute time from editor boot to PIE
-					FStudioTelemetryEditor::RecordEvent_Loading(TEXT("TimeToPIE"), EditorStartupTime + LoadMapTime + PIEStartupTime, Attributes);
-					FStudioTelemetryEditor::RecordEvent_CoreSystems(TEXT("TimeToPIE"), Attributes);
+			EditorLoadMapSpan->AddAttributes(Attributes);
 
-					IsFirstTimeToPIE = false;
-				}
-				
-				// Record the time from start PIE to PIE
-				FStudioTelemetryEditor::RecordEvent_Loading(TEXT("PIE.TotalStartupTime"), PIEStartupTime, Attributes);
-				FStudioTelemetryEditor::RecordEvent_CoreSystems(TEXT("PIE.TotalStartupTime"), Attributes);
-			}
+			FStudioTelemetry::Get().EndSpan(EditorLoadMapSpan);
 
-			PIELoadMapTime = FPlatformTime::Seconds() - PIELoadMapStartTime;
-
-			// Record the sub map load time
-			FStudioTelemetryEditor::RecordEvent_Loading(TEXT("PIE.LoadMapTime"), PIELoadMapTime, Attributes);
-			FStudioTelemetryEditor::RecordEvent_CoreSystems(TEXT("PIE.LoadMapTime"), Attributes);
-
-			TSharedPtr<FAnalyticsFlowTracker> FlowTracker = FStudioTelemetry::Get().GetFlowTracker().Pin();
-
-			if (FlowTracker.IsValid())
-			{
-				FlowTracker->EndSubFlow(PIELoadMapSubFlowGuid, true, Attributes);
-
-				// Start tracking World Streaming
-				if (UWorld* World = GameInstance->GetWorld())
-				{
-					FGuid WorldStreamingSubFlowGuid = FlowTracker->StartSubFlow(TEXT("World Streaming"));
-
-					WorldStreamingStartTime = FPlatformTime::Seconds();
-
-					World->OnWorldMatchStarting.AddLambda([this, WorldStreamingSubFlowGuid, MapName, FlowTracker]()
-						{
-							TArray<FAnalyticsEventAttribute> Attributes;
-							Attributes.Emplace(FAnalyticsEventAttribute(TEXT("LevelName"), LevelName));
-							Attributes.Emplace(FAnalyticsEventAttribute(TEXT("MapName"), MapName));
-
-							FlowTracker->EndSubFlow(WorldStreamingSubFlowGuid, true, Attributes);
-							
-							FStudioTelemetryEditor::RecordEvent_Loading(TEXT("WorldStreaming"), FPlatformTime::Seconds() - WorldStreamingStartTime, Attributes);
-						});
-				}
-			}
-		});
-
-	FEditorFileUtils::GetOnLoadMapStartDelegate().AddLambda([this]()
-		{
-			TSharedPtr<FAnalyticsFlowTracker> FlowTracker = FStudioTelemetry::Get().GetFlowTracker().Pin();
-
-			if (FlowTracker.IsValid())
-			{
-				// Load Map SubFlow will inherit the current Flow scope
-				LoadMapStartTime = FPlatformTime::Seconds();
-				LoadMapSubFlowGuid = FlowTracker->StartSubFlow(TEXT("Load Map"));
-			}
-		});
-
-	FEditorFileUtils::GetOnLoadMapEndDelegate().AddLambda([this](const FString& MapName)
-		{
-			LevelName = MapName;
-
-			TArray<FAnalyticsEventAttribute> Attributes;
-			Attributes.Emplace(FAnalyticsEventAttribute(TEXT("LevelName"), LevelName));
-			Attributes.Emplace(FAnalyticsEventAttribute(TEXT("MapName"), MapName));
-
-			LoadMapTime = FPlatformTime::Seconds() - LoadMapStartTime;
-			
-			FStudioTelemetryEditor::RecordEvent_Loading(TEXT("LoadMap"), LoadMapTime, Attributes);
-			FStudioTelemetryEditor::RecordEvent_CoreSystems(TEXT("LoadMap"), Attributes);
-
-			TSharedPtr<FAnalyticsFlowTracker> FlowTracker = FStudioTelemetry::Get().GetFlowTracker().Pin();
-
-			if (FlowTracker.IsValid())
-			{
-				FlowTracker->EndSubFlow(LoadMapSubFlowGuid, true, Attributes);
-			}
+			FStudioTelemetryEditor::RecordEvent_Loading(TEXT("LoadMap"), EditorLoadMapSpan->GetDuration(), EditorLoadMapSpan->GetAttributes());
+			FStudioTelemetryEditor::RecordEvent_CoreSystems(TEXT("LoadMap"), EditorLoadMapSpan->GetAttributes());
 		});
 
 	FEditorDelegates::OnEditorBoot.AddLambda([this](double TimeToBootEditor)
-		{
-			// Editor Boot Flow has finished
-			TArray<FAnalyticsEventAttribute> Attributes;
-			Attributes.Emplace(FAnalyticsEventAttribute(TEXT("LevelName"), TEXT("EditorBoot")));
-			Attributes.Emplace(FAnalyticsEventAttribute(TEXT("MapName"), TEXT("EditorBoot")));
-			
-			FStudioTelemetryEditor::RecordEvent_Loading(TEXT("BootEditor"), TimeToBootEditor, Attributes);
-			FStudioTelemetryEditor::RecordEvent_CoreSystems(TEXT("BootEditor"), Attributes);
-
-			TSharedPtr<FAnalyticsFlowTracker> FlowTracker = FStudioTelemetry::Get().GetFlowTracker().Pin();
+		{	
+			FStudioTelemetry::Get().EndSpan(EditorBootSpan);
 
-			if (FlowTracker.IsValid())
-			{
-				FlowTracker->EndFlow(EditorBootFlowGuid, true, Attributes);
+			// Callback is received when the editor has booted but has not been initialized
+			FStudioTelemetryEditor::RecordEvent_Loading(TEXT("BootEditor"), EditorBootSpan->GetDuration(), EditorBootSpan->GetAttributes());
+			FStudioTelemetryEditor::RecordEvent_CoreSystems(TEXT("BootEditor"), EditorBootSpan->GetAttributes());
 
-				// Start the Editor Interact Flow
-				InteractiveEditorFlowGuid = FlowTracker->StartFlow(TEXT("Interactive Editor"));
-			}
+			EditorInitilizeSpan = FStudioTelemetry::Get().StartSpan(EditorInitilizeSpanName);
 		});
 
 	FEditorDelegates::OnEditorInitialized.AddLambda([this](double TimeToInitializeEditor)
 		{
-			EditorStartupTime = TimeToInitializeEditor;
+			TimeToStartEditor = TimeToInitializeEditor;
 
+			// Editor has finished initializing so start the Editor Interact span
+			FStudioTelemetry::Get().EndSpan(EditorInitilizeSpan);
+			EditorInteractSpan = FStudioTelemetry::Get().StartSpan(EditorInteractSpanName);
+			
 			// Editor has initialized
 			TArray<FAnalyticsEventAttribute> Attributes;
 			Attributes.Emplace(FAnalyticsEventAttribute(TEXT("LevelName"), TEXT("EditorInitialize")));
@@ -595,12 +468,7 @@ void FStudioTelemetryEditor::Initialize()
 				GEditor->GetEditorSubsystem<UAssetEditorSubsystem>()->OnAssetEditorRequestedOpen().AddLambda([this](UObject* Asset)
 					{
 						AssetOpenStartTime = FPlatformTime::Seconds();
-						TSharedPtr<FAnalyticsFlowTracker> FlowTracker = FStudioTelemetry::Get().GetFlowTracker().Pin();
-
-						if (FlowTracker.IsValid())
-						{
-							FlowTracker->StartSubFlow(TEXT("Open Asset Editor"));
-						}
+						FStudioTelemetry::Get().StartSpan(OpenAssetEditorSpan);
 					});
 
 				GEditor->GetEditorSubsystem<UAssetEditorSubsystem>()->OnAssetOpenedInEditor().AddLambda([this](UObject* Asset, IAssetEditorInstance*)
@@ -617,12 +485,7 @@ void FStudioTelemetryEditor::Initialize()
 							FStudioTelemetryEditor::RecordEvent_Loading(TEXT("OpenAssetEditor"), FPlatformTime::Seconds() - AssetOpenStartTime, Attributes);
 						}
 
-						TSharedPtr<FAnalyticsFlowTracker> FlowTracker = FStudioTelemetry::Get().GetFlowTracker().Pin();
-
-						if (FlowTracker.IsValid())
-						{
-							FlowTracker->EndSubFlow(TEXT("Open Asset Editor"), true, Attributes);
-						}
+						FStudioTelemetry::Get().EndSpan(OpenAssetEditorSpan, Attributes);			
 				});
 			}
 
@@ -645,47 +508,128 @@ void FStudioTelemetryEditor::Initialize()
 
 			if (GWarn != nullptr)
 			{
+				// Start the SlowTask span
 				GWarn->OnStartSlowTask().AddLambda([this](const FText& TaskName)
 					{	
-						TSharedPtr<FAnalyticsFlowTracker> FlowTracker = FStudioTelemetry::Get().GetFlowTracker().Pin();
+						TArray<FAnalyticsEventAttribute> Attributes;
+						Attributes.Emplace(TEXT("TaskName"), TaskName.ToString());
 
-						if (FlowTracker.IsValid())
-						{
-							FlowTracker->StartSubFlow(FName(TaskName.ToString()));
-						}
+						FStudioTelemetry::Get().StartSpan(FName(TaskName.ToString()), Attributes);
 					});
 
+				// End the SlowTask span
 				GWarn->OnFinalizeSlowTask().AddLambda([this](const FText& TaskName, double TaskDurationSeconds)
 					{
-						TArray<FAnalyticsEventAttribute> Attributes;
-						Attributes.Emplace(FAnalyticsEventAttribute(TEXT("LevelName"), LevelName));
-						Attributes.Emplace(TEXT("TaskName"), TaskName.ToString());
-						
-						TSharedPtr<FAnalyticsFlowTracker> FlowTracker = FStudioTelemetry::Get().GetFlowTracker().Pin();
+						TSharedPtr<IAnalyticsSpan> SlowTaskSpan = FStudioTelemetry::Get().GetSpan(FName(TaskName.ToString()));
 
-						if (FlowTracker.IsValid())
+						if (SlowTaskSpan.IsValid())
 						{
-							FlowTracker->EndSubFlow(FName(TaskName.ToString()), true, Attributes);
+							FStudioTelemetry::Get().EndSpan(SlowTaskSpan);
+							FStudioTelemetryEditor::RecordEvent_Loading(TEXT("SlowTask"), SlowTaskSpan->GetDuration(), SlowTaskSpan->GetAttributes());	
 						}
-						
-						FStudioTelemetryEditor::RecordEvent_Loading(TEXT("SlowTaskDialog"), TaskDurationSeconds, Attributes);
 					});
 			}
 		});
+
+	// Install PIE Mode callbacks
+	FEditorDelegates::BeginPIE.AddLambda([this](bool)
+		{
+			// PIE mode has been started. The user has pressed the Start PIE button.
+			// Finish the Editor span
+			FStudioTelemetry::Get().EndSpan(EditorSpan);
+
+			// Start PIE span
+			PIESpan = FStudioTelemetry::Get().StartSpan(PIESpanName);
+			PIEStartupSpan = FStudioTelemetry::Get().StartSpan(PIEStartupSpanName);		
+		});
+
+	FWorldDelegates::OnPIEMapCreated.AddLambda([this](UGameInstance* GameInstance)
+		{
+			// A new PIE map was created
+			PIELoadMapSpan = FStudioTelemetry::Get().StartSpan(PIELoadMapSpanName);
+		});
+
+	FWorldDelegates::OnPIEMapReady.AddLambda([this](UGameInstance* GameInstance)
+		{
+			// PIE map is now loaded and ready to use
+			const FString MapName = FPaths::GetBaseFilename(GameInstance->PIEMapName);
+
+			TArray<FAnalyticsEventAttribute> Attributes;
+			Attributes.Emplace(FAnalyticsEventAttribute(TEXT("LevelName"), LevelName));
+			Attributes.Emplace(FAnalyticsEventAttribute(TEXT("MapName"), MapName));
+
+			PIELoadMapSpan->AddAttributes(Attributes);
+			
+			FStudioTelemetry::Get().EndSpan(PIELoadMapSpan);
+
+			FStudioTelemetryEditor::RecordEvent_Loading(TEXT("PIE.LoadMapTime"), PIELoadMapSpan->GetDuration(), PIELoadMapSpan->GetAttributes());
+			FStudioTelemetryEditor::RecordEvent_CoreSystems(TEXT("PIE.LoadMapTime"), PIELoadMapSpan->GetAttributes());
+		});
+
+	FWorldDelegates::OnPIEReady.AddLambda([this](UGameInstance* GameInstance)
+		{
+			// PIE is now ready for user interaction
+			static bool IsFirstTimeToPIE = true;
+
+			FStudioTelemetry::Get().EndSpan(PIEStartupSpan);
+
+			// Record the time from start PIE to PIE
+			FStudioTelemetryEditor::RecordEvent_Loading(TEXT("PIE.TotalStartupTime"), PIEStartupSpan->GetDuration(), PIEStartupSpan->GetAttributes());
+			FStudioTelemetryEditor::RecordEvent_CoreSystems(TEXT("PIE.TotalStartupTime"), PIEStartupSpan->GetAttributes());
+
+			if (IsFirstTimeToPIE == true)
+			{
+				// Record the absolute time from editor boot to PIE
+				FStudioTelemetryEditor::RecordEvent_Loading(TEXT("TimeToPIE"), TimeToStartEditor + EditorLoadMapSpan->GetDuration() + PIEStartupSpan->GetDuration(), PIEStartupSpan->GetAttributes());
+				FStudioTelemetryEditor::RecordEvent_CoreSystems(TEXT("TimeToPIE"), PIEStartupSpan->GetAttributes());
+
+				IsFirstTimeToPIE = false;
+			}
+	
+			// Start PIE World Streaming span
+			if (GameInstance)
+			{
+				if (UWorld* World = GameInstance->GetWorld())
+				{
+					PIEWorldStreamingSpan = FStudioTelemetry::Get().StartSpan(PIEWorldStreamingSpanName);
+
+					World->OnWorldMatchStarting.AddLambda([this]()
+						{
+							FStudioTelemetry::Get().EndSpan(PIEWorldStreamingSpan);
+							FStudioTelemetryEditor::RecordEvent_Loading(TEXT("PIE.WorldStreaming"), PIEWorldStreamingSpan->GetDuration(), PIEWorldStreamingSpan->GetAttributes());	
+						});
+				}
+			}	
+		});
+
+	FEditorDelegates::EndPIE.AddLambda([this](bool)
+		{
+			// PIE has ended, ie. the user has pressed the Stop PIE button, and we are going back to interactive Editor mode	
+			FStudioTelemetry::Get().EndSpan(PIESpan);
+
+			// Restart the Editor span
+			EditorSpan = FStudioTelemetry::Get().StartSpan(EditorSpanName);
+			EditorInteractSpan = FStudioTelemetry::Get().StartSpan(EditorInteractSpanName);
+		});
+
 	
 	// Install Cooking Callbacks
 	UE::Cook::FDelegates::CookByTheBookStarted.AddLambda([this](UE::Cook::ICookInfo& CookInfo)
 	{
-			TSharedPtr<FAnalyticsFlowTracker> FlowTracker = FStudioTelemetry::Get().GetFlowTracker().Pin();
+		// Begin the cooking span	
+		CookingSpan = FStudioTelemetry::Get().StartSpan(TEXT("Cooking"));
 
-			if (FlowTracker.IsValid())
-			{
-				CookByTheBookFlowGuid = FlowTracker->StartFlow("Cook By The Book");
-			}
+		TArray<FAnalyticsEventAttribute> Attributes;
+		Attributes.Emplace(FAnalyticsEventAttribute(TEXT("LevelName"), LevelName));
+		Attributes.Emplace(FAnalyticsEventAttribute(TEXT("MapName"), LevelName));
+
+		CookingSpan->AddAttributes(Attributes);
 	});
 
 	UE::Cook::FDelegates::CookByTheBookFinished.AddLambda([this](UE::Cook::ICookInfo& CookInfo)
 	{
+		// End the cooking span
+	
 		// Suppress sending telemetry from CookWorkers for now.
 		uint32 MultiprocessId = 0;
 		FParse::Value(FCommandLine::Get(), TEXT("-MultiprocessId="), MultiprocessId);
@@ -694,21 +638,13 @@ void FStudioTelemetryEditor::Initialize()
 			return;
 		}
 
-		TArray<FAnalyticsEventAttribute> Attributes;
-		Attributes.Emplace(FAnalyticsEventAttribute(TEXT("LevelName"), LevelName));
-		Attributes.Emplace(FAnalyticsEventAttribute(TEXT("MapName"), LevelName));
-
-		TSharedPtr<FAnalyticsFlowTracker> FlowTracker = FStudioTelemetry::Get().GetFlowTracker().Pin();
-
-		if (FlowTracker.IsValid())
-		{
-			FlowTracker->EndFlow(CookByTheBookFlowGuid, true, Attributes);
-		}
+		FStudioTelemetryEditor::RecordEvent_Cooking(CookingSpan->GetAttributes());
+		FStudioTelemetryEditor::RecordEvent_CoreSystems(TEXT("Cooking"), CookingSpan->GetAttributes());
 
-		FStudioTelemetryEditor::RecordEvent_Cooking(Attributes);
-		FStudioTelemetryEditor::RecordEvent_CoreSystems(TEXT("Cooking"), Attributes);
+		FStudioTelemetry::Get().EndSpan(CookingSpan);
 	});
 
+	// Install Content Browser callbacks
 	FContentBrowserModule* ContentBrowserModule = FModuleManager::GetModulePtr<FContentBrowserModule>( TEXT("ContentBrowser") );
 	
 	FTelemetryRouter& Router = FTelemetryRouter::Get();
diff --git a/Engine/Plugins/Experimental/StudioTelemetry/Source/StudioTelemetry/Private/StudioTelemetryEditor.h b/Engine/Plugins/Experimental/StudioTelemetry/Source/StudioTelemetry/Private/StudioTelemetryEditor.h
--- a/Engine/Plugins/Experimental/StudioTelemetry/Source/StudioTelemetry/Private/StudioTelemetryEditor.h
+++ b/Engine/Plugins/Experimental/StudioTelemetry/Source/StudioTelemetry/Private/StudioTelemetryEditor.h
@@ -34,26 +34,33 @@ private:
 	static void RecordEvent_VirtualAssets(const FString& Context, TArray<FAnalyticsEventAttribute> Attributes = {});
 	static void RegisterCollectionWorkflowDelegates(FTelemetryRouter& Router);
 
-	/** Internal variables for Flows and Timer data */
-	FString LevelName;
-	FGuid EditorBootFlowGuid;
-	FGuid InteractiveEditorFlowGuid;
-	FGuid LoadMapSubFlowGuid;
-	FGuid PIEFlowGuid;
-	FGuid PIEInitializeSubFlowGuid;
-	FGuid PIELoadMapSubFlowGuid;
-	FGuid CookByTheBookFlowGuid;
+	TSharedPtr<IAnalyticsSpan> EditorSpan;
+	TSharedPtr<IAnalyticsSpan> EditorBootSpan;
+	TSharedPtr<IAnalyticsSpan> EditorInteractSpan;
+	TSharedPtr<IAnalyticsSpan> EditorInitilizeSpan;
+	TSharedPtr<IAnalyticsSpan> EditorLoadMapSpan;
+	TSharedPtr<IAnalyticsSpan> PIESpan;
+	TSharedPtr<IAnalyticsSpan> PIEStartupSpan;
+	TSharedPtr<IAnalyticsSpan> PIELoadMapSpan;
+	TSharedPtr<IAnalyticsSpan> PIEWorldStreamingSpan;
+	TSharedPtr<IAnalyticsSpan> CookingSpan;
+	
+	const FName EditorSpanName = TEXT("Editor");
+	const FName EditorBootSpanName = TEXT("Editor.Boot");
+	const FName EditorInitilizeSpanName = TEXT("Editor.Initialize");
+	const FName EditorInteractSpanName = TEXT("Editor.Interact");
+	const FName EditorLoadMapSpanName = TEXT("Editor.LoadMap");
+	const FName PIESpanName = TEXT("PIE");
+	const FName PIEStartupSpanName = TEXT("PIE.Startup");
+	const FName PIELoadMapSpanName = TEXT("PIE.LoadMap");
+	const FName PIEWorldStreamingSpanName = TEXT("PIE.WorldStreaming");
+	const FName CookingSpanName = TEXT("Cooking");
+	const FName OpenAssetEditorSpan = TEXT("Open Asset Editor");
 
+	FString LevelName;
 	double SessionStartTime;
-	double PIEStartTime;
-	double LoadMapStartTime;
 	double AssetOpenStartTime;
-	double WorldStreamingStartTime;
-	double EditorStartupTime=0;
-	double LoadMapTime=0;
-	double PIEStartupTime = 0;
-	double PIELoadMapTime = 0;
-	double PIELoadMapStartTime;
+	double TimeToStartEditor;
 };
 
 #endif // WITH_EDITOR
diff --git a/Engine/Plugins/Experimental/StudioTelemetry/Source/StudioTelemetry/Public/StudioTelemetry.h b/Engine/Plugins/Experimental/StudioTelemetry/Source/StudioTelemetry/Public/StudioTelemetry.h
--- a/Engine/Plugins/Experimental/StudioTelemetry/Source/StudioTelemetry/Public/StudioTelemetry.h
+++ b/Engine/Plugins/Experimental/StudioTelemetry/Source/StudioTelemetry/Public/StudioTelemetry.h
@@ -3,7 +3,7 @@
 #pragma once
 
 #include "Interfaces/IAnalyticsProvider.h"
-#include "AnalyticsFlowTracker.h"
+#include "Interfaces/IAnalyticsTracer.h"
 #include "Containers/Array.h"
 #include "Containers/UnrealString.h"
 #include "CoreMinimal.h"
@@ -11,9 +11,16 @@
 #include "Templates/SharedPointer.h"
 
 class FAnalyticsProviderMulticast;
-
 /**
- * Public facing Studio Telemetry Plugin API
+ * Studio Telemetry Plugin API
+ * 
+ * Notes:
+ * Telemetry for Common Editor and Core Engine is collected automatically.
+ * Telemetry Sessions are started and ended automatically with the plugin initialization and shutdown. As such telemetry will not be captured prior to the plugin initialization.
+ * Developers are encouraged to add their own telemetry via this API or to intercept the event recording via the supplied callback on the SetRecordEventCallback API below.
+ * It is strongly recommended that developers implement their own IAnalyticsProviderModule where custom recording of telemetry events is desired.
+ * Custom AnalyticsProviders can be added to the plugin via the .ini. See FAnalyticsProviderLog or FAnalyticsProviderET for example.
+ * Telemetry events are recored to all registered IAnalyticsProviders supplied in the .ini file using the FAnalyticsProviderMulticast provider, except where specifically recorded with the RecordEvent(ProviderName,.. ) API below
  */
 class FStudioTelemetry : public IModuleInterface
 {
@@ -33,16 +40,31 @@ public:
 	/** Access to the broadcast analytics provider for the system*/
 	STUDIOTELEMETRY_API TWeakPtr<IAnalyticsProvider> GetProvider();
 
-	/** Access to the flow tracker for the system, this will ultimately broadcast flow events to all providers*/
-	STUDIOTELEMETRY_API TWeakPtr<FAnalyticsFlowTracker> GetFlowTracker();
+	/** Access to the tracer for the system*/
+	STUDIOTELEMETRY_API TWeakPtr<IAnalyticsTracer> GetTracer();
 	
-	/** Thread safe method to record an event to all analytics providers*/
+	/** Thread safe method to record an event to all registered analytics providers*/
 	STUDIOTELEMETRY_API void RecordEvent(const FString& EventName, const TArray<FAnalyticsEventAttribute>& Attributes = {});
 
-	/** Thread safe method to record an event to the named analytics provider */
+	/** Thread safe method to record an event to the specifically named analytics provider */
 	STUDIOTELEMETRY_API void RecordEvent(const FString& ProviderName, const FString& EventName, const TArray<FAnalyticsEventAttribute>& Attributes = {});
-	
-	/** Method for custom recording of telemetry events*/
+
+	/** Start a new span specifying the parent*/
+	STUDIOTELEMETRY_API TSharedPtr<IAnalyticsSpan> StartSpan(const FName Name, const TArray<FAnalyticsEventAttribute>& AdditionalAttributes = {});
+
+	/** Start a new span specifying the parent*/
+	STUDIOTELEMETRY_API TSharedPtr<IAnalyticsSpan> StartSpan(const FName Name, TSharedPtr<IAnalyticsSpan> ParentSpan, const TArray<FAnalyticsEventAttribute>& AdditionalAttributes = {});
+
+	/** End an existing span*/
+	STUDIOTELEMETRY_API bool EndSpan(TSharedPtr<IAnalyticsSpan> Span, const TArray<FAnalyticsEventAttribute>& AdditionalAttributes = {});
+
+	/** End an existing span by name*/
+	STUDIOTELEMETRY_API bool EndSpan(const FName& Name, const TArray<FAnalyticsEventAttribute>& AdditionalAttributes = {});
+
+	/** Get an active span by name, non active spans will not be available*/
+	STUDIOTELEMETRY_API TSharedPtr<IAnalyticsSpan> GetSpan(const FName Name);
+
+	/** Callback for interception of telemetry events recording that can be used by Developers to send telemetry events to their own back end, though it is recommended that Developers implement their own IAnalyticsProvider via their own IAnalyticsProviderModule*/
 	STUDIOTELEMETRY_API void SetRecordEventCallback(OnRecordEvent);
 
 private:
@@ -59,6 +81,6 @@ private:
 
 	FCriticalSection						CriticalSection;
 	TSharedPtr<FAnalyticsProviderMulticast>	AnalyticsProvider;
-	TSharedPtr<FAnalyticsFlowTracker>		AnalyticsFlowTracker;
+	TSharedPtr<IAnalyticsTracer>			AnalyticsTracer;
 	OnRecordEvent							RecordEventCallback;
 };
diff --git a/Engine/Source/Runtime/Analytics/Analytics/Private/Analytics.cpp b/Engine/Source/Runtime/Analytics/Analytics/Private/Analytics.cpp
--- a/Engine/Source/Runtime/Analytics/Analytics/Private/Analytics.cpp
+++ b/Engine/Source/Runtime/Analytics/Analytics/Private/Analytics.cpp
@@ -1,6 +1,7 @@
 // Copyright Epic Games, Inc. All Rights Reserved.
 
 #include "Analytics.h"
+#include "AnalyticsTracer.h"
 #include "Misc/CommandLine.h"
 #include "Interfaces/IAnalyticsProviderModule.h"
 
@@ -40,6 +41,11 @@ TSharedPtr<IAnalyticsProvider> FAnalytics::CreateAnalyticsProvider( const FName&
 	return NULL;
 }
 
+TSharedPtr<IAnalyticsTracer> FAnalytics::CreateAnalyticsTracer()
+{
+	return MakeShared<FAnalyticsTracer>();
+}
+
 FString FAnalytics::GetConfigValueFromIni( const FString& IniName, const FString& SectionName, const FString& KeyName, bool bIsRequired )
 {
 	FString Result;
diff --git a/Engine/Source/Runtime/Analytics/Analytics/Private/AnalyticsTracer.cpp b/Engine/Source/Runtime/Analytics/Analytics/Private/AnalyticsTracer.cpp
new file mode 100644
--- /dev/null
+++ b/Engine/Source/Runtime/Analytics/Analytics/Private/AnalyticsTracer.cpp
@@ -0,0 +1,248 @@
+// Copyright Epic Games, Inc. All Rights Reserved.
+
+#include "AnalyticsTracer.h"
+#include "ProfilingDebugging/MiscTrace.h"
+
+UE_DISABLE_OPTIMIZATION_SHIP
+ 
+static void AggregateAttributes(TArray<FAnalyticsEventAttribute>& AggregatedAttibutes, const TArray<FAnalyticsEventAttribute>& Attributes)
+{
+	// Aggregates all attributes
+	for (const FAnalyticsEventAttribute& Attribute : Attributes)
+	{
+		bool AttributeWasFound = false;
+
+		for (FAnalyticsEventAttribute& AggregatedAttribute : AggregatedAttibutes)
+		{
+			if (Attribute.GetName() == AggregatedAttribute.GetName())
+			{
+				AggregatedAttribute += Attribute;
+
+				// If we already have this attribute then great no more to do for this attribute
+				AttributeWasFound = true;
+				break;
+			}
+		}
+
+		if (AttributeWasFound == false)
+		{
+			// No matching attribute so append
+			AggregatedAttibutes.Add(Attribute);
+		}
+	}
+}
+
+void FAnalyticsSpan::SetProvider(TSharedPtr<IAnalyticsProvider> Provider)
+{
+	AnalyticsProvider = Provider;
+}
+
+double FAnalyticsSpan::GetDuration() const 
+{
+	return Duration;
+}
+
+void FAnalyticsSpan::Start(const FName& NewSpanName, TSharedPtr<IAnalyticsSpan> NewSpanParent, const TArray<FAnalyticsEventAttribute>& AdditionalAttributes)
+{
+	// Create a new Guid for this flow, can we assume it is unique?
+	Name			= NewSpanName;
+	Guid			= FGuid::NewGuid();
+	Attributes		= AdditionalAttributes;
+	ThreadId		= FPlatformTLS::GetCurrentThreadId();
+	StartTime		= FDateTime::UtcNow();
+	EndTime			= FDateTime::UtcNow();
+	Duration		= 0;
+	ScopeDepth		= NewSpanParent.IsValid()? NewSpanParent->GetScopeDepth() + 1 : 0;
+	ParentSpan		= NewSpanParent;
+	IsActive		= true;
+
+	TRACE_BEGIN_REGION(*Name.ToString());
+}
+
+void FAnalyticsSpan::AddChildSpan(TSharedPtr<IAnalyticsSpan> ChildSpan)
+{
+	ChildSpans.Add(ChildSpan);
+}
+
+void FAnalyticsSpan::End(const TArray<FAnalyticsEventAttribute>& AdditionalAttributes)
+{
+	// Only End the span once
+	if (IsActive == false)
+	{
+		return;
+	}
+
+	// Calculate the duration
+	EndTime = FDateTime::UtcNow();
+	Duration = (EndTime - StartTime).GetTotalSeconds();
+		
+	TRACE_END_REGION(*Name.ToString());
+
+	// Append the parent attributes and the the additional attributes to the current span attributes, these will get passed down to the child spans
+	AddAttributes(AdditionalAttributes);
+
+	TSharedPtr< IAnalyticsSpan> ParentSpanShared = ParentSpan.Pin();
+
+	if (ParentSpanShared.IsValid())
+	{
+		AddAttributes(ParentSpanShared->GetAttributes());
+	}
+
+	const uint32 SpanSchemaVersion = 1;
+	const FString SpanEventName = TEXT("Span");
+
+	TArray<FAnalyticsEventAttribute> EventAttributes = Attributes;
+
+	EventAttributes.Add(FAnalyticsEventAttribute(TEXT("SchemaVersion"), SpanSchemaVersion));
+	EventAttributes.Add(FAnalyticsEventAttribute(TEXT("Span_Name"), Name.ToString()));
+	EventAttributes.Add(FAnalyticsEventAttribute(TEXT("Span_GUID"), Guid.ToString()));
+	EventAttributes.Add(FAnalyticsEventAttribute(TEXT("Span_ParentName"), ParentSpanShared.IsValid() ? ParentSpanShared->GetName().ToString() : TEXT("")));
+	EventAttributes.Add(FAnalyticsEventAttribute(TEXT("Span_ThreadId"), ThreadId));
+	EventAttributes.Add(FAnalyticsEventAttribute(TEXT("Span_Depth"), ScopeDepth));
+	EventAttributes.Add(FAnalyticsEventAttribute(TEXT("Span_StartUTC"), StartTime.ToUnixTimestampDecimal()));
+	EventAttributes.Add(FAnalyticsEventAttribute(TEXT("Span_EndUTC"), EndTime.ToUnixTimestampDecimal()));
+	EventAttributes.Add(FAnalyticsEventAttribute(TEXT("Span_TimeInSec"), Duration));
+
+	// Make sure we end all the child spans 
+	for (TSharedPtr<IAnalyticsSpan> ChildSpan : ChildSpans)
+	{
+		ChildSpan->End(Attributes);
+	}
+
+	if (AnalyticsProvider.IsValid())
+	{
+		AnalyticsProvider->RecordEvent(SpanEventName, EventAttributes);
+	}
+
+	IsActive = false;
+}
+
+void FAnalyticsSpan::AddAttributes(const TArray<FAnalyticsEventAttribute>& AdditionalAttributes)
+{
+	AggregateAttributes(Attributes, AdditionalAttributes);
+}
+
+void FAnalyticsSpan::RecordEvent(const FString& EventName, const TArray<FAnalyticsEventAttribute>& AdditionalAttributes)
+{
+	if (AnalyticsProvider.IsValid())
+	{
+		TArray<FAnalyticsEventAttribute> EventAttributes = Attributes;
+		AggregateAttributes(EventAttributes, AdditionalAttributes);
+		AnalyticsProvider->RecordEvent(EventName, EventAttributes);
+	}
+}
+
+const FName& FAnalyticsSpan::GetName() const
+{
+	return Name;
+}
+
+const TArray<FAnalyticsEventAttribute>& FAnalyticsSpan::GetAttributes() const
+{
+	return Attributes;
+}
+
+uint32 FAnalyticsSpan::GetScopeDepth() const
+{
+	return ScopeDepth;
+}
+
+TSharedPtr<IAnalyticsSpan> FAnalyticsSpan::GetParentSpan() const
+{
+	return ParentSpan.Pin();
+}
+
+void FAnalyticsTracer::SetProvider(TSharedPtr<IAnalyticsProvider> InProvider)
+{
+	AnalyticsProvider = InProvider;
+}
+
+void FAnalyticsTracer::SetCurrentSpan(TSharedPtr<IAnalyticsSpan> Span)
+{
+	CurrentSpan = Span;
+}
+
+TSharedPtr<IAnalyticsSpan> FAnalyticsTracer::GetCurrentSpan() const
+{
+	return CurrentSpan;
+}
+
+void FAnalyticsTracer::StartSession()
+{
+	SessionSpan = StartSpan(TEXT("Session"), TSharedPtr<IAnalyticsSpan>());
+	SetCurrentSpan(SessionSpan);
+}
+
+void FAnalyticsTracer::EndSession()
+{
+	FScopeLock ScopeLock(&CriticalSection);	
+	SessionSpan->End();
+	SessionSpan.Reset();
+	AnalyticsProvider.Reset();
+}
+
+TSharedPtr<IAnalyticsSpan> FAnalyticsTracer::StartSpan(const FName NewSpanName, TSharedPtr<IAnalyticsSpan> ParentSpan, const TArray<FAnalyticsEventAttribute>& AdditionalAttributes)
+{
+	FScopeLock ScopeLock(&CriticalSection);
+	return StartSpanInternal(NewSpanName, ParentSpan.IsValid()? ParentSpan : GetCurrentSpan(), AdditionalAttributes);
+}
+
+TSharedPtr<IAnalyticsSpan> FAnalyticsTracer::StartSpanInternal(const FName NewSpanName, TSharedPtr<IAnalyticsSpan> ParentSpan, const TArray<FAnalyticsEventAttribute>& AdditionalAttributes)
+{
+	TSharedPtr<IAnalyticsSpan> NewSpan = MakeShared<FAnalyticsSpan>();
+	
+	NewSpan->SetProvider(AnalyticsProvider);
+	NewSpan->Start(NewSpanName, ParentSpan, AdditionalAttributes);
+
+	if (ParentSpan.IsValid())
+	{
+		ParentSpan->AddChildSpan(NewSpan);
+	}
+
+	SpanRegistry.Emplace(NewSpanName, NewSpan);
+
+	SetCurrentSpan(NewSpan);
+
+	return NewSpan;
+}
+
+bool FAnalyticsTracer::EndSpan(TSharedPtr<IAnalyticsSpan> Span, const TArray<FAnalyticsEventAttribute>& AdditionalAttributes)
+{
+	FScopeLock ScopeLock(&CriticalSection);
+	return EndSpanInternal(Span, AdditionalAttributes);
+}
+
+bool FAnalyticsTracer::EndSpanInternal(TSharedPtr<IAnalyticsSpan> Span, const TArray<FAnalyticsEventAttribute>& AdditionalAttributes)
+{
+	if (Span.IsValid())
+	{
+		Span->End(AdditionalAttributes);
+
+		SetCurrentSpan(Span->GetParentSpan());
+
+		SpanRegistry.Remove(Span->GetName());
+
+		return true;
+	}
+
+	return false;
+}
+
+TSharedPtr<IAnalyticsSpan> FAnalyticsTracer::GetSessionSpan() const
+{
+	return SessionSpan;
+}
+
+TSharedPtr<IAnalyticsSpan> FAnalyticsTracer::GetSpanInternal(const FName Name)
+{
+	TWeakPtr<IAnalyticsSpan>* SpanWeakPtr = SpanRegistry.Find(Name);
+	return SpanWeakPtr != nullptr ? (*SpanWeakPtr).Pin() : TSharedPtr<IAnalyticsSpan>();
+}
+
+TSharedPtr<IAnalyticsSpan> FAnalyticsTracer::GetSpan(const FName Name)
+{
+	FScopeLock ScopeLock(&CriticalSection);
+	return GetSpanInternal(Name);
+}
+
+UE_ENABLE_OPTIMIZATION_SHIP
\ No newline at end of file
diff --git a/Engine/Source/Runtime/Analytics/Analytics/Public/Analytics.h b/Engine/Source/Runtime/Analytics/Analytics/Public/Analytics.h
--- a/Engine/Source/Runtime/Analytics/Analytics/Public/Analytics.h
+++ b/Engine/Source/Runtime/Analytics/Analytics/Public/Analytics.h
@@ -19,6 +19,7 @@
 ANALYTICS_API DECLARE_LOG_CATEGORY_EXTERN(LogAnalytics, Display, All);
 
 class IAnalyticsProvider;
+class IAnalyticsTracer;
 
 /**
  * The public interface for interacting with analytics.
@@ -75,6 +76,7 @@ public:
 	 */
 	virtual TSharedPtr<IAnalyticsProvider> CreateAnalyticsProvider(const FName& ProviderModuleName, const FAnalyticsProviderConfigurationDelegate& GetConfigValue);
 
+	virtual TSharedPtr<IAnalyticsTracer> CreateAnalyticsTracer();
 	/**
 	 * Creates an instance of the default configured analytics provider.
 	 * Default is determined by GetDefaultProviderModuleName and a default constructed ConfigFromIni instance.
diff --git a/Engine/Source/Runtime/Analytics/Analytics/Public/AnalyticsTracer.h b/Engine/Source/Runtime/Analytics/Analytics/Public/AnalyticsTracer.h
new file mode 100644
--- /dev/null
+++ b/Engine/Source/Runtime/Analytics/Analytics/Public/AnalyticsTracer.h
@@ -0,0 +1,86 @@
+// Copyright Epic Games, Inc. All Rights Reserved.
+
+#pragma once
+
+#include "Interfaces/IAnalyticsTracer.h"
+
+using FThreadId = uint32;
+
+/**
+* Implementation of an IAnalyticsSpan interface
+*/
+class FAnalyticsSpan : public IAnalyticsSpan
+{
+public:
+
+	FAnalyticsSpan(){};
+	~FAnalyticsSpan(){};
+
+	// Public IAnalyticsSpan implementation
+	virtual void SetProvider(TSharedPtr<IAnalyticsProvider> AnalyticsProvider) override;
+	virtual void Start(const FName& Name, TSharedPtr<IAnalyticsSpan> ParentSpan, const TArray<FAnalyticsEventAttribute>& AdditionalAttributes = {});
+	virtual void End(const TArray<FAnalyticsEventAttribute>& AdditionalAttributes = {}) override;
+	virtual void AddAttributes(const TArray<FAnalyticsEventAttribute>& AdditionalAttributes) override;
+	virtual void RecordEvent(const FString& EventName, const TArray<FAnalyticsEventAttribute>& AdditionalAttributes = {}) override;
+	virtual const FName& GetName() const override;
+	virtual const TArray<FAnalyticsEventAttribute>& GetAttributes() const override;
+	virtual uint32 GetScopeDepth() const override;
+	virtual TSharedPtr<IAnalyticsSpan> GetParentSpan() const override;
+	virtual double GetDuration() const override;
+	virtual void AddChildSpan(TSharedPtr<IAnalyticsSpan> ChildSpan) override;
+
+private:
+
+	FName								Name = TEXT("None");
+	FGuid								Guid = FGuid();
+	FDateTime							StartTime = 0;
+	FDateTime							EndTime = 0;
+	FThreadId							ThreadId = 0;
+	uint32								ScopeDepth = 0;
+	double								Duration = 0;
+	bool								IsActive = false;
+	TSharedPtr<IAnalyticsProvider>		AnalyticsProvider;
+	TArray<FAnalyticsEventAttribute>	Attributes;
+	TWeakPtr<IAnalyticsSpan>			ParentSpan;
+	TArray<TSharedPtr<IAnalyticsSpan>>	ChildSpans;
+};
+
+/**
+* Implementation of an IAnalyticsTracer interface
+*/
+class FAnalyticsTracer : public IAnalyticsTracer
+{
+public:
+
+	FAnalyticsTracer(){};
+	~FAnalyticsTracer(){};
+	
+	// Public IAnalyticsTracer implementation
+	virtual void StartSession() override;
+	virtual void EndSession() override;
+	virtual void SetProvider(TSharedPtr<IAnalyticsProvider> AnalyticsProvider) override;
+	virtual TSharedPtr<IAnalyticsSpan> StartSpan(const FName Name, TSharedPtr<IAnalyticsSpan> ParentSpan, const TArray<FAnalyticsEventAttribute>& AdditionalAttributes = {}) override;
+	virtual bool EndSpan(TSharedPtr<IAnalyticsSpan> Span, const TArray<FAnalyticsEventAttribute>& AdditionalAttributes = {}) override;
+	virtual TSharedPtr<IAnalyticsSpan> GetCurrentSpan() const override;
+	virtual TSharedPtr<IAnalyticsSpan> GetSessionSpan() const override;
+	virtual TSharedPtr<IAnalyticsSpan> GetSpan(const FName Name) override;
+
+protected:
+
+	// Protected IAnalyticsTracer implementation
+	virtual void SetCurrentSpan(TSharedPtr<IAnalyticsSpan> Span) override;
+
+private:
+
+	TSharedPtr<IAnalyticsSpan> StartSpanInternal(const FName Name, TSharedPtr<IAnalyticsSpan> ParentSpan, const TArray<FAnalyticsEventAttribute>& AdditionalAttributes);
+	bool EndSpanInternal(TSharedPtr<IAnalyticsSpan> Span, const TArray<FAnalyticsEventAttribute>& AdditionalAttributes);
+	TSharedPtr<IAnalyticsSpan> GetSpanInternal(const FName Name);
+
+	TSharedPtr<IAnalyticsProvider>			AnalyticsProvider;
+	TSharedPtr<IAnalyticsSpan>				SessionSpan;
+	TSharedPtr<IAnalyticsSpan>				CurrentSpan;
+	TMap<FName,TWeakPtr<IAnalyticsSpan>>	SpanRegistry;
+	FCriticalSection						CriticalSection;
+};
+
+
diff --git a/Engine/Source/Runtime/Analytics/Analytics/Public/Interfaces/IAnalyticsTracer.h b/Engine/Source/Runtime/Analytics/Analytics/Public/Interfaces/IAnalyticsTracer.h
new file mode 100644
--- /dev/null
+++ b/Engine/Source/Runtime/Analytics/Analytics/Public/Interfaces/IAnalyticsTracer.h
@@ -0,0 +1,119 @@
+// Copyright Epic Games, Inc. All Rights Reserved.
+
+#pragma once
+
+#include "UObject/NameTypes.h"
+#include "UObject/UnrealNames.h"
+#include "IAnalyticsProvider.h"
+
+/**
+ * IAnalyticsSpan Interface 
+ * 
+ * Inspired by OpenTelemety model for Tracer/Span.
+ * A span should be used to measure long running actions. It can be thought of as hierarchical timer object.
+ * A span capture event state within an array of EventAattributes and emit them with every event sent from this span
+ * A span can only be started via a IAnalyticsTracer interface, but can be Ended on demand
+ * Spans are arranged hierarchically and the expected behavior is for the parent span to end the spans of its children on during it's own End call
+ * Elapsed time should be calculated on each call, but on when the span is active
+ * Will send telemetry events via the IAnalyticsProvider
+ */
+class IAnalyticsSpan : FNoncopyable
+{
+public:
+
+	IAnalyticsSpan() {};
+	virtual ~IAnalyticsSpan() {};
+
+	/** Sets the analytics provider for the tracer */
+	virtual void SetProvider(TSharedPtr<IAnalyticsProvider> AnalyticsProvider) = 0;
+
+	/** Start this span */
+	virtual void Start(const FName& Name, TSharedPtr<IAnalyticsSpan> ParentSpan, const TArray<FAnalyticsEventAttribute>& AdditionalAttributes = {}) = 0;
+
+	/** End this span */
+	virtual void End(const TArray<FAnalyticsEventAttribute>& AdditionalAttributes = {}) = 0;
+
+	/** Append attributes to this span context */
+	virtual void AddAttributes(const TArray<FAnalyticsEventAttribute>& AdditionalAttributes) = 0;
+	
+	/** Record an event from this span, appends the attributes for the span context */
+	virtual void RecordEvent(const FString& EventName, const TArray<FAnalyticsEventAttribute>& AdditionalAttributes = {}) = 0;
+
+	/** Get the internal name for the span */
+	virtual const FName& GetName() const = 0;
+
+	/** Get context attributes for the span */
+	virtual const TArray<FAnalyticsEventAttribute>& GetAttributes()const = 0;
+
+	/** Get the scope depth span */
+	virtual uint32 GetScopeDepth() const = 0;
+
+	/** Get the span duration in seconds */
+	virtual double GetDuration() const = 0;
+
+	/** Get the parent span */
+	virtual TSharedPtr<IAnalyticsSpan> GetParentSpan() const = 0;
+
+	// Add a child span
+	virtual void AddChildSpan(TSharedPtr<IAnalyticsSpan> ChildSpan) = 0;
+};
+
+/**
+ * IAnalyticsTracer Interface
+ *
+ * Inspired by OpenTelemety model for Tracer/Span.
+ * A tracer manages and create spans. It is not permitted to create or start a span outside of the owning tracer object. 
+ * By default, new spans are added as a child of the last added span but can be added to a child of a specified parent for clarity
+ * The tracer has a valid SessionSpan object which is the root of the span hierarchy which is expected to be valid between the StartSession and EndSession calls.
+ * Spans are hierarchical and are not expected to overlap, even though the API or implementation might well allow that.
+ */
+class IAnalyticsTracer : FNoncopyable
+{
+public:
+	IAnalyticsTracer() {};
+	virtual ~IAnalyticsTracer() {};
+
+	// Starts a new session, creates and Starts the SessionSpan object
+	virtual void StartSession() =0;
+
+	// Ends the running session and Ends the SessionSpan object
+	virtual void EndSession() = 0;
+
+	/** Sets the analytics provider for the tracer */
+	virtual void SetProvider(TSharedPtr<IAnalyticsProvider> AnalyticsProvider) =0;
+
+	/** Start a new span specifying an optional parent, null parent will attach new span under the current span */
+	virtual TSharedPtr<IAnalyticsSpan> StartSpan(const FName Name, TSharedPtr<IAnalyticsSpan> ParentSpan=TSharedPtr<IAnalyticsSpan>(), const TArray<FAnalyticsEventAttribute>& AdditionalAttributes = {}) = 0;
+
+	/** End an existing span*/
+	virtual bool EndSpan(TSharedPtr<IAnalyticsSpan>, const TArray<FAnalyticsEventAttribute>& AdditionalAttributes = {}) = 0;
+
+	/** Get the currently active span */
+	virtual TSharedPtr<IAnalyticsSpan> GetCurrentSpan() const = 0;
+
+	/** Get the session span, this will always be valid for an active session */
+	virtual TSharedPtr<IAnalyticsSpan> GetSessionSpan() const = 0;
+
+	/** Get an active span by name, non active spans will not be available*/
+	virtual TSharedPtr<IAnalyticsSpan> GetSpan(const FName Name) = 0;
+
+	/** Start a new span without a parent*/
+	TSharedPtr<IAnalyticsSpan> StartSpan(const FName Name, const TArray<FAnalyticsEventAttribute>& AdditionalAttributes = {})
+	{
+		return StartSpan(Name, GetCurrentSpan(), AdditionalAttributes);
+	};
+
+	/** End an existing span by name*/
+	bool EndSpan(const FName& Name, const TArray<FAnalyticsEventAttribute>& AdditionalAttributes = {})
+	{
+		return EndSpan(GetSpan(Name), AdditionalAttributes);
+	}
+
+protected:
+	/** Set the currently active span */
+	virtual void SetCurrentSpan(TSharedPtr<IAnalyticsSpan> Span) = 0;
+
+private:
+};
+
+
diff --git a/Engine/Source/Runtime/Analytics/AnalyticsET/Public/AnalyticsFlowTracker.h b/Engine/Source/Runtime/Analytics/AnalyticsET/Public/AnalyticsFlowTracker.h
--- a/Engine/Source/Runtime/Analytics/AnalyticsET/Public/AnalyticsFlowTracker.h
+++ b/Engine/Source/Runtime/Analytics/AnalyticsET/Public/AnalyticsFlowTracker.h
@@ -8,7 +8,7 @@
 
 using FThreadId = uint32;
 
-class FAnalyticsFlowTracker : FNoncopyable
+class UE_DEPRECATED(5.4, "Use FAnalyticsTracer instead.") FAnalyticsFlowTracker : FNoncopyable
 {
 public:
 	FAnalyticsFlowTracker() {};
diff --git a/Engine/Source/Runtime/Engine/Classes/Engine/GameInstance.h b/Engine/Source/Runtime/Engine/Classes/Engine/GameInstance.h
--- a/Engine/Source/Runtime/Engine/Classes/Engine/GameInstance.h
+++ b/Engine/Source/Runtime/Engine/Classes/Engine/GameInstance.h
@@ -294,7 +294,7 @@ public:
 	/** Called as soon as the game mode is spawned, to allow additional PIE setting validation prior to creating the local players / etc... (called on pure clients too, in which case the game mode is nullptr) */
 	ENGINE_API virtual FGameInstancePIEResult PostCreateGameModeForPIE(const FGameInstancePIEParameters& Params, AGameModeBase* GameMode);
 
-	ENGINE_API void ReportPIEStartupTime();
+	ENGINE_API virtual void ReportPIEStartupTime();
 #endif
 
 	ENGINE_API class UEngine* GetEngine() const;
diff --git a/Engine/Source/Runtime/Engine/Classes/Engine/World.h b/Engine/Source/Runtime/Engine/Classes/Engine/World.h
--- a/Engine/Source/Runtime/Engine/Classes/Engine/World.h
+++ b/Engine/Source/Runtime/Engine/Classes/Engine/World.h
@@ -4200,6 +4200,14 @@ public:
 	DECLARE_MULTICAST_DELEGATE_OneParam(FOnWorldPIEReady, UGameInstance*);
 	static FOnWorldPIEStarted OnPIEReady;
 
+	// PIE map is created
+	DECLARE_MULTICAST_DELEGATE_OneParam(FOnWorldPIEMapCreated, UGameInstance*);
+	static FOnWorldPIEMapCreated OnPIEMapCreated;
+
+	// PIE map is created
+	DECLARE_MULTICAST_DELEGATE_OneParam(FOnWorldPIEMapReady, UGameInstance*);
+	static FOnWorldPIEMapReady OnPIEMapReady;
+
 	// PIE has ended
 	DECLARE_MULTICAST_DELEGATE_OneParam(FOnWorldPIEEnded, UGameInstance*);
 	static FOnWorldPIEEnded OnPIEEnded;
diff --git a/Engine/Source/Runtime/Engine/Private/GameInstance.cpp b/Engine/Source/Runtime/Engine/Private/GameInstance.cpp
--- a/Engine/Source/Runtime/Engine/Private/GameInstance.cpp
+++ b/Engine/Source/Runtime/Engine/Private/GameInstance.cpp
@@ -278,6 +278,7 @@ static ENetMode GetNetModeFromPlayNetMode(const EPlayNetMode InPlayNetMode, cons
 FGameInstancePIEResult UGameInstance::InitializeForPlayInEditor(int32 PIEInstanceIndex, const FGameInstancePIEParameters& Params)
 {
 	FWorldDelegates::OnPIEStarted.Broadcast(this);
+	FWorldDelegates::OnPIEMapCreated.Broadcast(this);
 
 	UEditorEngine* const EditorEngine = CastChecked<UEditorEngine>(GetEngine());
 
@@ -364,14 +365,15 @@ FGameInstancePIEResult UGameInstance::InitializeForPlayInEditor(int32 PIEInstanc
 	// This creates the world subsystems and prepares to begin play
 	EditorEngine->PostCreatePIEWorld(NewWorld);
 
+	FWorldDelegates::OnPIEMapReady.Broadcast(this);
+
 	// Games can override this to return failure if PIE is not allowed for some reason
 	return FGameInstancePIEResult::Success();
 }
 
-
+#if WITH_EDITOR
 void UGameInstance::ReportPIEStartupTime()
 {
-#if WITH_EDITOR
 	if (!bReportedPIEStartupTime)
 	{
 		static bool bHasRunPIEThisSession = false;
@@ -388,9 +390,8 @@ void UGameInstance::ReportPIEStartupTime()
 	}
 
 	FWorldDelegates::OnPIEReady.Broadcast(this);
-#endif
 }
-
+#endif
 
 FGameInstancePIEResult UGameInstance::StartPlayInEditorGameInstance(ULocalPlayer* LocalPlayer, const FGameInstancePIEParameters& Params)
 {
diff --git a/Engine/Source/Runtime/Engine/Private/World.cpp b/Engine/Source/Runtime/Engine/Private/World.cpp
--- a/Engine/Source/Runtime/Engine/Private/World.cpp
+++ b/Engine/Source/Runtime/Engine/Private/World.cpp
@@ -577,6 +577,8 @@ FWorldDelegates::FWorldEvent FWorldDelegates::OnWorldBeginTearDown;
 FWorldDelegates::FRefreshLevelScriptActionsEvent FWorldDelegates::RefreshLevelScriptActions;
 FWorldDelegates::FOnWorldPIEStarted FWorldDelegates::OnPIEStarted;
 FWorldDelegates::FOnWorldPIEReady FWorldDelegates::OnPIEReady;
+FWorldDelegates::FOnWorldPIEMapCreated FWorldDelegates::OnPIEMapCreated;
+FWorldDelegates::FOnWorldPIEMapReady FWorldDelegates::OnPIEMapReady;
 FWorldDelegates::FOnWorldPIEEnded FWorldDelegates::OnPIEEnded;
 #endif // WITH_EDITOR
 FWorldDelegates::FOnSeamlessTravelStart FWorldDelegates::OnSeamlessTravelStart;
