Fix for SlowTask spans so they do not use their actual name as the span name

diff --git a/Engine/Plugins/Experimental/StudioTelemetry/Source/StudioTelemetry/Private/StudioTelemetryEditor.cpp b/Engine/Plugins/Experimental/StudioTelemetry/Source/StudioTelemetry/Private/StudioTelemetryEditor.cpp
--- a/Engine/Plugins/Experimental/StudioTelemetry/Source/StudioTelemetry/Private/StudioTelemetryEditor.cpp
+++ b/Engine/Plugins/Experimental/StudioTelemetry/Source/StudioTelemetry/Private/StudioTelemetryEditor.cpp
@@ -511,21 +511,48 @@ void FStudioTelemetryEditor::Initialize()
 				// Start the SlowTask span
 				GWarn->OnStartSlowTask().AddLambda([this](const FText& TaskName)
 					{	
-						TArray<FAnalyticsEventAttribute> Attributes;
-						Attributes.Emplace(TEXT("TaskName"), TaskName.ToString());
+						// Slow tasks can possibly be started from multiple threads, so we need to protect the registered span table
+						FScopeLock ScopeLock(&TaskSpanCriticalSection);
+						
+						FString Name = TaskName.ToString();
+
+						TSharedPtr<IAnalyticsSpan>* SpanPtr = TaskSpans.Find(Name);
+						ensureMsgf(SpanPtr == nullptr, TEXT("We assume that only one task with the name %s is in flight"), *Name);
+						
+						if (TaskSpans.Find(Name) == SpanPtr)
+						{
+							TArray<FAnalyticsEventAttribute> Attributes;
+							Attributes.Emplace(TEXT("TaskName"), TaskName.ToString());
+
+							// Create and start a new slow task span
+							TSharedPtr<IAnalyticsSpan> SlowTaskSpan = FStudioTelemetry::Get().StartSpan(TEXT("SlowTask"), Attributes);
 
-						FStudioTelemetry::Get().StartSpan(FName(TaskName.ToString()), Attributes);
+							// Store this SlowTask span so we can find it when it finishes
+							TaskSpans.Add(Name, SlowTaskSpan);
+						}
 					});
 
 				// End the SlowTask span
 				GWarn->OnFinalizeSlowTask().AddLambda([this](const FText& TaskName, double TaskDurationSeconds)
 					{
-						TSharedPtr<IAnalyticsSpan> SlowTaskSpan = FStudioTelemetry::Get().GetSpan(FName(TaskName.ToString()));
+						// Slow tasks can possibly be finalized from multiple threads, so we need to protect the registered span table
+						FScopeLock ScopeLock(&TaskSpanCriticalSection);
+
+						FString Name = TaskName.ToString();
+
+						// Find the task we stored off when we started this task
+						TSharedPtr<IAnalyticsSpan>* SpanPtr = TaskSpans.Find(Name);
+						ensureMsgf(SpanPtr != nullptr, TEXT("Unable to find a registered span with name %s"), *Name);
 
-						if (SlowTaskSpan.IsValid())
+						if (SpanPtr!=nullptr)
 						{
+							TSharedPtr<IAnalyticsSpan> SlowTaskSpan = *SpanPtr;
+
 							FStudioTelemetry::Get().EndSpan(SlowTaskSpan);
 							FStudioTelemetryEditor::RecordEvent_Loading(TEXT("SlowTask"), SlowTaskSpan->GetDuration(), SlowTaskSpan->GetAttributes());	
+
+							// Remove the SlowTask span from the registry
+							TaskSpans.Remove(Name);
 						}
 					});
 			}
@@ -587,7 +614,7 @@ void FStudioTelemetryEditor::Initialize()
 			}
 	
 			// Start PIE World Streaming span
-			if (GameInstance)
+			/*if (GameInstance)
 			{
 				if (UWorld* World = GameInstance->GetWorld())
 				{
@@ -596,10 +623,10 @@ void FStudioTelemetryEditor::Initialize()
 					World->OnWorldMatchStarting.AddLambda([this]()
 						{
 							FStudioTelemetry::Get().EndSpan(PIEWorldStreamingSpan);
-							FStudioTelemetryEditor::RecordEvent_Loading(TEXT("PIE.WorldStreaming"), PIEWorldStreamingSpan->GetDuration(), PIEWorldStreamingSpan->GetAttributes());	
+							FStudioTelemetryEditor::RecordEvent_Loading(TEXT("PIE.WorldStreaming"), PIEWorldStreamingSpan->GetDuration(), PIEWorldStreamingSpan->GetAttributes());
 						});
 				}
-			}	
+			}	*/
 		});
 
 	FEditorDelegates::EndPIE.AddLambda([this](bool)
diff --git a/Engine/Plugins/Experimental/StudioTelemetry/Source/StudioTelemetry/Private/StudioTelemetryEditor.h b/Engine/Plugins/Experimental/StudioTelemetry/Source/StudioTelemetry/Private/StudioTelemetryEditor.h
--- a/Engine/Plugins/Experimental/StudioTelemetry/Source/StudioTelemetry/Private/StudioTelemetryEditor.h
+++ b/Engine/Plugins/Experimental/StudioTelemetry/Source/StudioTelemetry/Private/StudioTelemetryEditor.h
@@ -57,6 +57,9 @@ private:
 	const FName CookingSpanName = TEXT("Cooking");
 	const FName OpenAssetEditorSpan = TEXT("Open Asset Editor");
 
+	TMap<FString, TSharedPtr<IAnalyticsSpan>> TaskSpans;
+	FCriticalSection TaskSpanCriticalSection;
+
 	FString LevelName;
 	double SessionStartTime;
 	double AssetOpenStartTime;
