Minor tweak to telemetry data so that map names appear in PIE Startup spans

diff --git a/Engine/Plugins/Experimental/StudioTelemetry/Source/StudioTelemetry/Private/StudioTelemetryEditor.cpp b/Engine/Plugins/Experimental/StudioTelemetry/Source/StudioTelemetry/Private/StudioTelemetryEditor.cpp
--- a/Engine/Plugins/Experimental/StudioTelemetry/Source/StudioTelemetry/Private/StudioTelemetryEditor.cpp
+++ b/Engine/Plugins/Experimental/StudioTelemetry/Source/StudioTelemetry/Private/StudioTelemetryEditor.cpp
@@ -568,6 +568,11 @@ void FStudioTelemetryEditor::Initialize()
 			// Start PIE span
 			PIESpan = FStudioTelemetry::Get().StartSpan(PIESpanName);
 			PIEStartupSpan = FStudioTelemetry::Get().StartSpan(PIEStartupSpanName);		
+
+			TArray<FAnalyticsEventAttribute> Attributes;
+			Attributes.Emplace(FAnalyticsEventAttribute(TEXT("LevelName"), LevelName));
+
+			PIESpan->AddAttributes(Attributes);
 		});
 
 	FWorldDelegates::OnPIEMapCreated.AddLambda([this](UGameInstance* GameInstance)
