Fix for SlowTask span creation.
Fixed the obvious logic error with the spanptr test. Removed the ensure as it is likely, and allowed, for multiple SlowTasks with the same name to run ashynchronously. This is not  currently supported by the telemetry system,

diff --git a/Engine/Plugins/Experimental/StudioTelemetry/Source/StudioTelemetry/Private/StudioTelemetryEditor.cpp b/Engine/Plugins/Experimental/StudioTelemetry/Source/StudioTelemetry/Private/StudioTelemetryEditor.cpp
--- a/Engine/Plugins/Experimental/StudioTelemetry/Source/StudioTelemetry/Private/StudioTelemetryEditor.cpp
+++ b/Engine/Plugins/Experimental/StudioTelemetry/Source/StudioTelemetry/Private/StudioTelemetryEditor.cpp
@@ -515,11 +515,10 @@ void FStudioTelemetryEditor::Initialize()
 						FScopeLock ScopeLock(&TaskSpanCriticalSection);
 						
 						FString Name = TaskName.ToString();
-
 						TSharedPtr<IAnalyticsSpan>* SpanPtr = TaskSpans.Find(Name);
-						ensureMsgf(SpanPtr == nullptr, TEXT("We assume that only one task with the name %s is in flight"), *Name);
 						
-						if (TaskSpans.Find(Name) == SpanPtr)
+						// Only one task with this name running asynchronously is supported at this time.
+						if (SpanPtr==nullptr)
 						{
 							TArray<FAnalyticsEventAttribute> Attributes;
 							Attributes.Emplace(TEXT("TaskName"), TaskName.ToString());
@@ -542,8 +541,7 @@ void FStudioTelemetryEditor::Initialize()
 
 						// Find the task we stored off when we started this task
 						TSharedPtr<IAnalyticsSpan>* SpanPtr = TaskSpans.Find(Name);
-						ensureMsgf(SpanPtr != nullptr, TEXT("Unable to find a registered span with name %s"), *Name);
-
+						
 						if (SpanPtr!=nullptr)
 						{
 							TSharedPtr<IAnalyticsSpan> SlowTaskSpan = *SpanPtr;
