Fxied typo in FStudioTelemetry.IsAvailable
Renamed LevelName -> MapName
Remaned MapName -> PIE.MapName

diff --git a/Engine/Plugins/Experimental/StudioTelemetry/Source/StudioTelemetry/Private/StudioTelemetryEditor.cpp b/Engine/Plugins/Experimental/StudioTelemetry/Source/StudioTelemetry/Private/StudioTelemetryEditor.cpp
--- a/Engine/Plugins/Experimental/StudioTelemetry/Source/StudioTelemetry/Private/StudioTelemetryEditor.cpp
+++ b/Engine/Plugins/Experimental/StudioTelemetry/Source/StudioTelemetry/Private/StudioTelemetryEditor.cpp
@@ -403,10 +403,10 @@ void FStudioTelemetryEditor::Initialize()
 	// Start Editor and Editor Boot span. Note : this will only start when the plugin is loaded and as such will miss any activity that runs beforehand
 	EditorSpan = FStudioTelemetry::Get().StartSpan(EditorSpanName);
 	EditorBootSpan = FStudioTelemetry::Get().StartSpan(EditorBootSpanName);
+	EditorMapName = TEXT("None");
 
 	TArray<FAnalyticsEventAttribute> Attributes;
-	Attributes.Emplace(FAnalyticsEventAttribute(TEXT("LevelName"), TEXT("EditorBoot")));
-	Attributes.Emplace(FAnalyticsEventAttribute(TEXT("MapName"), TEXT("EditorBoot")));
+	Attributes.Emplace(FAnalyticsEventAttribute(TEXT("MapName"), EditorMapName));
 
 	EditorBootSpan->AddAttributes(Attributes);
 
@@ -419,11 +419,10 @@ void FStudioTelemetryEditor::Initialize()
 	FEditorDelegates::OnMapOpened.AddLambda([this](const FString& MapName, bool Unused)
 		{
 			// The new editor map was actually opened
-			LevelName = FPaths::GetBaseFilename(MapName);
+			EditorMapName = FPaths::GetBaseFilename(MapName);
 
 			TArray<FAnalyticsEventAttribute> Attributes;
-			Attributes.Emplace(FAnalyticsEventAttribute(TEXT("LevelName"), LevelName));
-			Attributes.Emplace(FAnalyticsEventAttribute(TEXT("MapName"), LevelName));
+			Attributes.Emplace(FAnalyticsEventAttribute(TEXT("MapName"), EditorMapName));
 
 			EditorLoadMapSpan->AddAttributes(Attributes);
 
@@ -454,8 +453,7 @@ void FStudioTelemetryEditor::Initialize()
 			
 			// Editor has initialized
 			TArray<FAnalyticsEventAttribute> Attributes;
-			Attributes.Emplace(FAnalyticsEventAttribute(TEXT("LevelName"), TEXT("EditorInitialize")));
-			Attributes.Emplace(FAnalyticsEventAttribute(TEXT("MapName"), TEXT("EditorInitialize")));
+			Attributes.Emplace(FAnalyticsEventAttribute(TEXT("MapName"), EditorMapName));
 			
 			FStudioTelemetryEditor::RecordEvent_Loading(TEXT("TotalEditorStartup"), TimeToInitializeEditor, Attributes);
 			FStudioTelemetryEditor::RecordEvent_CoreSystems(TEXT("TotalEditorStartup"), Attributes);
@@ -474,11 +472,11 @@ void FStudioTelemetryEditor::Initialize()
 				GEditor->GetEditorSubsystem<UAssetEditorSubsystem>()->OnAssetOpenedInEditor().AddLambda([this](UObject* Asset, IAssetEditorInstance*)
 					{
 						TArray<FAnalyticsEventAttribute> Attributes;
-						Attributes.Emplace(FAnalyticsEventAttribute(TEXT("LevelName"), LevelName));
+						Attributes.Emplace(FAnalyticsEventAttribute(TEXT("MapName"), EditorMapName));
 						
 						if (Asset != nullptr)
 						{
-							Attributes.Emplace(FAnalyticsEventAttribute(TEXT("MapName"), Asset->GetFullName()));
+							Attributes.Emplace(FAnalyticsEventAttribute(TEXT("AssetName"), Asset->GetFullName()));
 							Attributes.Emplace(TEXT("AssetPath"), Asset->GetFullName());
 							Attributes.Emplace(TEXT("AssetType"), Asset->GetClass()->GetName());
 
@@ -496,8 +494,7 @@ void FStudioTelemetryEditor::Initialize()
 				UE::Cook::FDelegates::CookByTheBookFinished.AddLambda([this](UE::Cook::ICookInfo& CookInfo)
 					{
 						TArray<FAnalyticsEventAttribute> Attributes;
-						Attributes.Emplace(FAnalyticsEventAttribute(TEXT("LevelName"), LevelName));
-						Attributes.Emplace(FAnalyticsEventAttribute(TEXT("MapName"), LevelName));
+						Attributes.Emplace(FAnalyticsEventAttribute(TEXT("MapName"), EditorMapName));
 						
 						FStudioTelemetryEditor::RecordEvent_Cooking(Attributes);
 						FStudioTelemetryEditor::RecordEvent_CoreSystems(TEXT("Cooking"), Attributes);
@@ -568,7 +565,7 @@ void FStudioTelemetryEditor::Initialize()
 			PIEStartupSpan = FStudioTelemetry::Get().StartSpan(PIEStartupSpanName);		
 
 			TArray<FAnalyticsEventAttribute> Attributes;
-			Attributes.Emplace(FAnalyticsEventAttribute(TEXT("LevelName"), LevelName));
+			Attributes.Emplace(FAnalyticsEventAttribute(TEXT("MapName"), EditorMapName));
 
 			PIESpan->AddAttributes(Attributes);
 		});
@@ -582,11 +579,11 @@ void FStudioTelemetryEditor::Initialize()
 	FWorldDelegates::OnPIEMapReady.AddLambda([this](UGameInstance* GameInstance)
 		{
 			// PIE map is now loaded and ready to use
-			const FString MapName = FPaths::GetBaseFilename(GameInstance->PIEMapName);
+			PIEMapName = FPaths::GetBaseFilename(GameInstance->PIEMapName);
 
 			TArray<FAnalyticsEventAttribute> Attributes;
-			Attributes.Emplace(FAnalyticsEventAttribute(TEXT("LevelName"), LevelName));
-			Attributes.Emplace(FAnalyticsEventAttribute(TEXT("MapName"), MapName));
+			Attributes.Emplace(FAnalyticsEventAttribute(TEXT("MapName"), EditorMapName));
+			Attributes.Emplace(FAnalyticsEventAttribute(TEXT("PIE_MapName"), PIEMapName));
 
 			PIELoadMapSpan->AddAttributes(Attributes);
 			
@@ -650,8 +647,7 @@ void FStudioTelemetryEditor::Initialize()
 		CookingSpan = FStudioTelemetry::Get().StartSpan(TEXT("Cooking"));
 
 		TArray<FAnalyticsEventAttribute> Attributes;
-		Attributes.Emplace(FAnalyticsEventAttribute(TEXT("LevelName"), LevelName));
-		Attributes.Emplace(FAnalyticsEventAttribute(TEXT("MapName"), LevelName));
+		Attributes.Emplace(FAnalyticsEventAttribute(TEXT("MapName"), EditorMapName));
 
 		CookingSpan->AddAttributes(Attributes);
 	});
diff --git a/Engine/Plugins/Experimental/StudioTelemetry/Source/StudioTelemetry/Private/StudioTelemetryEditor.h b/Engine/Plugins/Experimental/StudioTelemetry/Source/StudioTelemetry/Private/StudioTelemetryEditor.h
--- a/Engine/Plugins/Experimental/StudioTelemetry/Source/StudioTelemetry/Private/StudioTelemetryEditor.h
+++ b/Engine/Plugins/Experimental/StudioTelemetry/Source/StudioTelemetry/Private/StudioTelemetryEditor.h
@@ -60,7 +60,8 @@ private:
 	TMap<FString, TSharedPtr<IAnalyticsSpan>> TaskSpans;
 	FCriticalSection TaskSpanCriticalSection;
 
-	FString LevelName;
+	FString EditorMapName;
+	FString PIEMapName;
 	double SessionStartTime;
 	double AssetOpenStartTime;
 	double TimeToStartEditor;
diff --git a/Engine/Plugins/Experimental/StudioTelemetry/Source/StudioTelemetry/Public/StudioTelemetry.h b/Engine/Plugins/Experimental/StudioTelemetry/Source/StudioTelemetry/Public/StudioTelemetry.h
--- a/Engine/Plugins/Experimental/StudioTelemetry/Source/StudioTelemetry/Public/StudioTelemetry.h
+++ b/Engine/Plugins/Experimental/StudioTelemetry/Source/StudioTelemetry/Public/StudioTelemetry.h
@@ -29,7 +29,7 @@ public:
 	typedef TFunction<void(const FString& EventName, const TArray<FAnalyticsEventAttribute>& Attrs)> OnRecordEvent;
 
 	/** Check whether the module is available*/
-	static STUDIOTELEMETRY_API bool IsAvaliable() { return FModuleManager::Get().IsModuleLoaded("StudioTelemetry"); }
+	static STUDIOTELEMETRY_API bool IsAvailable() { return FModuleManager::Get().IsModuleLoaded("StudioTelemetry"); }
 
 	/** Access to the module singleton*/
 	static STUDIOTELEMETRY_API FStudioTelemetry& Get();
