Added MapName to SlowTask span

diff --git a/Engine/Plugins/Experimental/StudioTelemetry/Source/StudioTelemetry/Private/StudioTelemetryEditor.cpp b/Engine/Plugins/Experimental/StudioTelemetry/Source/StudioTelemetry/Private/StudioTelemetryEditor.cpp
--- a/Engine/Plugins/Experimental/StudioTelemetry/Source/StudioTelemetry/Private/StudioTelemetryEditor.cpp
+++ b/Engine/Plugins/Experimental/StudioTelemetry/Source/StudioTelemetry/Private/StudioTelemetryEditor.cpp
@@ -422,8 +422,9 @@ void FStudioTelemetryEditor::Initialize()
 			EditorMapName = FPaths::GetBaseFilename(MapName);
 
 			TArray<FAnalyticsEventAttribute> Attributes;
-			Attributes.Emplace(FAnalyticsEventAttribute(TEXT("MapName"), EditorMapName));
+			Attributes.Emplace(TEXT("MapName"), EditorMapName);
 
+			EditorSpan->AddAttributes(Attributes);
 			EditorLoadMapSpan->AddAttributes(Attributes);
 
 			FStudioTelemetry::Get().EndSpan(EditorLoadMapSpan);
@@ -447,14 +448,14 @@ void FStudioTelemetryEditor::Initialize()
 		{
 			TimeToStartEditor = TimeToInitializeEditor;
 
-			// Editor has finished initializing so start the Editor Interact span
-			FStudioTelemetry::Get().EndSpan(EditorInitilizeSpan);
-			EditorInteractSpan = FStudioTelemetry::Get().StartSpan(EditorInteractSpanName);
-			
 			// Editor has initialized
 			TArray<FAnalyticsEventAttribute> Attributes;
-			Attributes.Emplace(FAnalyticsEventAttribute(TEXT("MapName"), EditorMapName));
+			Attributes.Emplace(TEXT("MapName"), EditorMapName);
 			
+			// Editor has finished initializing so start the Editor Interact span
+			FStudioTelemetry::Get().EndSpan(EditorInitilizeSpan);
+			EditorInteractSpan = FStudioTelemetry::Get().StartSpan(EditorInteractSpanName);
+						
 			FStudioTelemetryEditor::RecordEvent_Loading(TEXT("TotalEditorStartup"), TimeToInitializeEditor, Attributes);
 			FStudioTelemetryEditor::RecordEvent_CoreSystems(TEXT("TotalEditorStartup"), Attributes);
 
@@ -472,11 +473,10 @@ void FStudioTelemetryEditor::Initialize()
 				GEditor->GetEditorSubsystem<UAssetEditorSubsystem>()->OnAssetOpenedInEditor().AddLambda([this](UObject* Asset, IAssetEditorInstance*)
 					{
 						TArray<FAnalyticsEventAttribute> Attributes;
-						Attributes.Emplace(FAnalyticsEventAttribute(TEXT("MapName"), EditorMapName));
+						Attributes.Emplace(TEXT("MapName"), EditorMapName);
 						
 						if (Asset != nullptr)
 						{
-							Attributes.Emplace(FAnalyticsEventAttribute(TEXT("AssetName"), Asset->GetFullName()));
 							Attributes.Emplace(TEXT("AssetPath"), Asset->GetFullName());
 							Attributes.Emplace(TEXT("AssetType"), Asset->GetClass()->GetName());
 
@@ -494,7 +494,7 @@ void FStudioTelemetryEditor::Initialize()
 				UE::Cook::FDelegates::CookByTheBookFinished.AddLambda([this](UE::Cook::ICookInfo& CookInfo)
 					{
 						TArray<FAnalyticsEventAttribute> Attributes;
-						Attributes.Emplace(FAnalyticsEventAttribute(TEXT("MapName"), EditorMapName));
+						Attributes.Emplace(TEXT("MapName"), EditorMapName);
 						
 						FStudioTelemetryEditor::RecordEvent_Cooking(Attributes);
 						FStudioTelemetryEditor::RecordEvent_CoreSystems(TEXT("Cooking"), Attributes);
@@ -518,6 +518,7 @@ void FStudioTelemetryEditor::Initialize()
 						if (SpanPtr==nullptr)
 						{
 							TArray<FAnalyticsEventAttribute> Attributes;
+							Attributes.Emplace(TEXT("MapName"), EditorMapName);
 							Attributes.Emplace(TEXT("TaskName"), TaskName.ToString());
 
 							// Create and start a new slow task span
@@ -565,7 +566,7 @@ void FStudioTelemetryEditor::Initialize()
 			PIEStartupSpan = FStudioTelemetry::Get().StartSpan(PIEStartupSpanName);		
 
 			TArray<FAnalyticsEventAttribute> Attributes;
-			Attributes.Emplace(FAnalyticsEventAttribute(TEXT("MapName"), EditorMapName));
+			Attributes.Emplace(TEXT("MapName"), EditorMapName);
 
 			PIESpan->AddAttributes(Attributes);
 		});
@@ -582,8 +583,8 @@ void FStudioTelemetryEditor::Initialize()
 			PIEMapName = FPaths::GetBaseFilename(GameInstance->PIEMapName);
 
 			TArray<FAnalyticsEventAttribute> Attributes;
-			Attributes.Emplace(FAnalyticsEventAttribute(TEXT("MapName"), EditorMapName));
-			Attributes.Emplace(FAnalyticsEventAttribute(TEXT("PIE_MapName"), PIEMapName));
+			Attributes.Emplace(TEXT("MapName"), EditorMapName);
+			Attributes.Emplace(TEXT("PIE_MapName"), PIEMapName);
 
 			PIELoadMapSpan->AddAttributes(Attributes);
 			
@@ -634,8 +635,11 @@ void FStudioTelemetryEditor::Initialize()
 			// PIE has ended, ie. the user has pressed the Stop PIE button, and we are going back to interactive Editor mode	
 			FStudioTelemetry::Get().EndSpan(PIESpan);
 
+			TArray<FAnalyticsEventAttribute> Attributes;
+			Attributes.Emplace(TEXT("MapName"), EditorMapName);
+
 			// Restart the Editor span
-			EditorSpan = FStudioTelemetry::Get().StartSpan(EditorSpanName);
+			EditorSpan = FStudioTelemetry::Get().StartSpan(EditorSpanName, Attributes);
 			EditorInteractSpan = FStudioTelemetry::Get().StartSpan(EditorInteractSpanName);
 		});
 
@@ -647,7 +651,7 @@ void FStudioTelemetryEditor::Initialize()
 		CookingSpan = FStudioTelemetry::Get().StartSpan(TEXT("Cooking"));
 
 		TArray<FAnalyticsEventAttribute> Attributes;
-		Attributes.Emplace(FAnalyticsEventAttribute(TEXT("MapName"), EditorMapName));
+		Attributes.Emplace(TEXT("MapName"), EditorMapName);
 
 		CookingSpan->AddAttributes(Attributes);
 	});
