Fix for duplicate AssetType Attribute in Editor Telemetry
Renamed the AssetType from the OpenBrowser event to AssetClass ( that's what it is anyway! )

diff --git a/Engine/Plugins/Experimental/StudioTelemetry/Source/StudioTelemetry/Private/StudioTelemetryEditor.cpp b/Engine/Plugins/Experimental/StudioTelemetry/Source/StudioTelemetry/Private/StudioTelemetryEditor.cpp
--- a/Engine/Plugins/Experimental/StudioTelemetry/Source/StudioTelemetry/Private/StudioTelemetryEditor.cpp
+++ b/Engine/Plugins/Experimental/StudioTelemetry/Source/StudioTelemetry/Private/StudioTelemetryEditor.cpp
@@ -478,7 +478,7 @@ void FStudioTelemetryEditor::Initialize()
 						if (Asset != nullptr)
 						{
 							Attributes.Emplace(TEXT("AssetPath"), Asset->GetFullName());
-							Attributes.Emplace(TEXT("AssetType"), Asset->GetClass()->GetName());
+							Attributes.Emplace(TEXT("AssetClass"), Asset->GetClass()->GetName());
 
 							FStudioTelemetryEditor::RecordEvent_Loading(TEXT("OpenAssetEditor"), FPlatformTime::Seconds() - AssetOpenStartTime, Attributes);
 						}
