Added Start/Finalize events to FFeedBackContext that pass a unique GUID for the task. This is more reliable for observers ( like Studio Telemetry ) to handle Task callbacks.
Fixed SpanTask name for Slow Tasks so that it plays nicely with Insights
Removed UserID, ApplicationCommandLine and Hardware_ComputerName as this is not considered EGPI safe.

diff --git a/Engine/Plugins/Experimental/StudioTelemetry/Source/StudioTelemetry/Private/StudioTelemetry.cpp b/Engine/Plugins/Experimental/StudioTelemetry/Source/StudioTelemetry/Private/StudioTelemetry.cpp
--- a/Engine/Plugins/Experimental/StudioTelemetry/Source/StudioTelemetry/Private/StudioTelemetry.cpp
+++ b/Engine/Plugins/Experimental/StudioTelemetry/Source/StudioTelemetry/Private/StudioTelemetry.cpp
@@ -119,10 +119,7 @@ void FStudioTelemetry::StartSession()
 		FString SessionLabel;
 		FParse::Value(FCommandLine::Get(), TEXT("SessionLabel="), SessionLabel);
 
-		// Set the default event attributes
-		DefaultEventAttributes.Emplace(TEXT("User_ID"), UserID);
-		DefaultEventAttributes.Emplace(TEXT("Application_Commandline"), FCommandLine::Get());
-
+		// Set the default event attributes	
 		DefaultEventAttributes.Emplace(TEXT("ProjectName"), ProjectName);
 		DefaultEventAttributes.Emplace(TEXT("ProjectID"), ProjectID);
 
@@ -137,7 +134,6 @@ void FStudioTelemetry::StartSession()
 		DefaultEventAttributes.Emplace(TEXT("Build_BranchName"), FApp::GetBranchName().ToLower());
 		DefaultEventAttributes.Emplace(TEXT("Build_Changelist"), BuildSettings::GetCurrentChangelist());
 
-		DefaultEventAttributes.Emplace(TEXT("Hardware_ComputerName"), ComputerName.ToLower());
 		DefaultEventAttributes.Emplace(TEXT("Hardware_Platform"), FString(FPlatformProperties::IniPlatformName()));
 		DefaultEventAttributes.Emplace(TEXT("Hardware_GPU"), GRHIAdapterName);
 		DefaultEventAttributes.Emplace(TEXT("Hardware_CPU"), FPlatformMisc::GetCPUBrand());
diff --git a/Engine/Plugins/Experimental/StudioTelemetry/Source/StudioTelemetry/Private/StudioTelemetryEditor.cpp b/Engine/Plugins/Experimental/StudioTelemetry/Source/StudioTelemetry/Private/StudioTelemetryEditor.cpp
--- a/Engine/Plugins/Experimental/StudioTelemetry/Source/StudioTelemetry/Private/StudioTelemetryEditor.cpp
+++ b/Engine/Plugins/Experimental/StudioTelemetry/Source/StudioTelemetry/Private/StudioTelemetryEditor.cpp
@@ -410,6 +410,55 @@ void FStudioTelemetryEditor::Initialize()
 
 	EditorBootSpan->AddAttributes(Attributes);
 
+	ensureMsgf(GWarn, TEXT("GWarn was not valid"));
+
+	if (GWarn != nullptr)
+	{
+		// Start the SlowTask span
+		GWarn->OnStartSlowTaskWithGuid().AddLambda([this](FGuid TaskGuid, const FText& TaskName)
+			{
+				// Slow tasks can possibly be started from multiple threads, so we need to protect the registered span table
+				FScopeLock ScopeLock(&TaskSpanCriticalSection);
+
+				TSharedPtr<IAnalyticsSpan>* SpanPtr = TaskSpans.Find(TaskGuid);
+
+				// Only one task with this Guid is running asynchronously is supported at this time.
+				if (SpanPtr == nullptr)
+				{
+					TArray<FAnalyticsEventAttribute> Attributes;
+					Attributes.Emplace(TEXT("MapName"), EditorMapName);
+					Attributes.Emplace(TEXT("TaskName"), TaskName.ToString());
+
+					// Create and start a new slow task span
+					const FName SpanName(TEXT("ST ") + TaskName.ToString());
+					TSharedPtr<IAnalyticsSpan> SlowTaskSpan = FStudioTelemetry::Get().StartSpan(SpanName, Attributes);
+
+					// Store this SlowTask span so we can find it when it finishes
+					TaskSpans.Add(TaskGuid, SlowTaskSpan);
+				}
+			});
+
+		// End the SlowTask span
+		GWarn->OnFinalizeSlowTaskWithGuid().AddLambda([this](FGuid TaskGuid, double TaskDuration)
+			{
+				// Slow tasks can possibly be finalized from multiple threads, so we need to protect the registered span table
+				FScopeLock ScopeLock(&TaskSpanCriticalSection);
+
+				// Find the task we stored off when we started this task
+				TSharedPtr<IAnalyticsSpan>* SpanPtr = TaskSpans.Find(TaskGuid);
+
+				if (SpanPtr != nullptr)
+				{
+					TSharedPtr<IAnalyticsSpan> SlowTaskSpan = *SpanPtr;
+
+					FStudioTelemetry::Get().EndSpan(SlowTaskSpan);
+
+					// Remove the SlowTask span from the registry
+					TaskSpans.Remove(TaskGuid);
+				}
+			});
+	}
+
 	FEditorDelegates::OnMapLoad.AddLambda([this](const FString& MapName, FCanLoadMap& OutCanLoadMap)
 		{
 			// The Editor loads a new map
@@ -501,56 +550,7 @@ void FStudioTelemetryEditor::Initialize()
 					});
 			}
 
-			ensureMsgf(GWarn, TEXT("GWarn was not valid"));
-
-			if (GWarn != nullptr)
-			{
-				// Start the SlowTask span
-				GWarn->OnStartSlowTask().AddLambda([this](const FText& TaskName)
-					{	
-						// Slow tasks can possibly be started from multiple threads, so we need to protect the registered span table
-						FScopeLock ScopeLock(&TaskSpanCriticalSection);
-						
-						FString Name = TaskName.ToString();
-						TSharedPtr<IAnalyticsSpan>* SpanPtr = TaskSpans.Find(Name);
-						
-						// Only one task with this name running asynchronously is supported at this time.
-						if (SpanPtr==nullptr)
-						{
-							TArray<FAnalyticsEventAttribute> Attributes;
-							Attributes.Emplace(TEXT("MapName"), EditorMapName);
-							Attributes.Emplace(TEXT("TaskName"), TaskName.ToString());
-
-							// Create and start a new slow task span
-							TSharedPtr<IAnalyticsSpan> SlowTaskSpan = FStudioTelemetry::Get().StartSpan(TEXT("SlowTask"), Attributes);
-
-							// Store this SlowTask span so we can find it when it finishes
-							TaskSpans.Add(Name, SlowTaskSpan);
-						}
-					});
-
-				// End the SlowTask span
-				GWarn->OnFinalizeSlowTask().AddLambda([this](const FText& TaskName, double TaskDurationSeconds)
-					{
-						// Slow tasks can possibly be finalized from multiple threads, so we need to protect the registered span table
-						FScopeLock ScopeLock(&TaskSpanCriticalSection);
-
-						FString Name = TaskName.ToString();
-
-						// Find the task we stored off when we started this task
-						TSharedPtr<IAnalyticsSpan>* SpanPtr = TaskSpans.Find(Name);
-						
-						if (SpanPtr!=nullptr)
-						{
-							TSharedPtr<IAnalyticsSpan> SlowTaskSpan = *SpanPtr;
-
-							FStudioTelemetry::Get().EndSpan(SlowTaskSpan);
-					
-							// Remove the SlowTask span from the registry
-							TaskSpans.Remove(Name);
-						}
-					});
-			}
+			
 		});
 
 	// Install PIE Mode callbacks
diff --git a/Engine/Plugins/Experimental/StudioTelemetry/Source/StudioTelemetry/Private/StudioTelemetryEditor.h b/Engine/Plugins/Experimental/StudioTelemetry/Source/StudioTelemetry/Private/StudioTelemetryEditor.h
--- a/Engine/Plugins/Experimental/StudioTelemetry/Source/StudioTelemetry/Private/StudioTelemetryEditor.h
+++ b/Engine/Plugins/Experimental/StudioTelemetry/Source/StudioTelemetry/Private/StudioTelemetryEditor.h
@@ -57,7 +57,7 @@ private:
 	const FName CookingSpanName = TEXT("Cooking");
 	const FName OpenAssetEditorSpan = TEXT("Open Asset Editor");
 
-	TMap<FString, TSharedPtr<IAnalyticsSpan>> TaskSpans;
+	TMap<FGuid, TSharedPtr<IAnalyticsSpan>> TaskSpans;
 	FCriticalSection TaskSpanCriticalSection;
 
 	FString EditorMapName;
diff --git a/Engine/Source/Runtime/Core/Private/Misc/FeedbackContext.cpp b/Engine/Source/Runtime/Core/Private/Misc/FeedbackContext.cpp
--- a/Engine/Source/Runtime/Core/Private/Misc/FeedbackContext.cpp
+++ b/Engine/Source/Runtime/Core/Private/Misc/FeedbackContext.cpp
@@ -188,7 +188,6 @@ void FFeedbackContext::UpdateUI()
 /**** Begin legacy API ****/
 void FFeedbackContext::BeginSlowTask( const FText& Task, bool ShowProgressDialog, bool bShowCancelButton )
 {
-	TaskName = Task;
 	ensure(IsInGameThread());
 
 	TUniquePtr<FSlowTask> NewScope(new FSlowTask(0, Task, true, *this));
diff --git a/Engine/Source/Runtime/Core/Public/Misc/FeedbackContext.h b/Engine/Source/Runtime/Core/Public/Misc/FeedbackContext.h
--- a/Engine/Source/Runtime/Core/Public/Misc/FeedbackContext.h
+++ b/Engine/Source/Runtime/Core/Public/Misc/FeedbackContext.h
@@ -15,6 +15,7 @@
 #include "Misc/ScopeLock.h"
 #include "Misc/SlowTask.h"
 #include "Misc/SlowTaskStack.h"
+#include "Misc/Guid.h"
 #include "Templates/SharedPointer.h"
 #include "Templates/UniquePtr.h"
 #include "Templates/UnrealTemplate.h"
@@ -54,6 +55,11 @@ public:
 	DECLARE_EVENT_TwoParams(FFeedbackContext, FOnFinalizeSlowTask, const FText& TaskName, double DurationInSeconds);
 	FOnFinalizeSlowTask& OnFinalizeSlowTask() { return FinalizeSlowTaskEvent; }
 
+	DECLARE_EVENT_TwoParams(FFeedbackContext, FOnStartSlowTaskWithGuid, FGuid Guid, const FText& TaskName);
+	FOnStartSlowTaskWithGuid& OnStartSlowTaskWithGuid() { return StartSlowTaskWithGuidEvent; }
+
+	DECLARE_EVENT_TwoParams(FFeedbackContext, FOnFinalizeSlowTaskWithGuid, FGuid Guid, double DurationInSeconds);
+	FOnFinalizeSlowTaskWithGuid& OnFinalizeSlowTaskWithGuid() { return FinalizeSlowTaskWithGuidEvent; }
 
 	/**** Legacy API - not deprecated as it's still in heavy use, but superceded by FScopedSlowTask ****/
 	CORE_API void BeginSlowTask( const FText& Task, bool ShowProgressDialog, bool bShowCancelButton=false );
@@ -71,9 +77,11 @@ protected:
 	virtual void StartSlowTask( const FText& Task, bool bShowCancelButton=false )
 	{
 		TaskName = Task;
+		TaskGuid = FGuid::NewGuid();
 		TaskStartTime = FPlatformTime::Seconds();
 		GIsSlowTask = true;
 		StartSlowTaskEvent.Broadcast(TaskName);
+		StartSlowTaskWithGuidEvent.Broadcast(TaskGuid, TaskName);
 	}
 
 	/**
@@ -81,7 +89,9 @@ protected:
 	 */
 	virtual void FinalizeSlowTask( )
 	{
-		FinalizeSlowTaskEvent.Broadcast(TaskName, FPlatformTime::Seconds() - TaskStartTime);
+		const double TaskDuration = FPlatformTime::Seconds() - TaskStartTime;
+		FinalizeSlowTaskEvent.Broadcast(TaskName, TaskDuration);
+		FinalizeSlowTaskWithGuidEvent.Broadcast(TaskGuid, TaskDuration);
 		GIsSlowTask = false;
 	}
 
@@ -164,9 +174,14 @@ private:
 
 	/** The name of any task we are running */
 	FText TaskName;
+
 	double TaskStartTime;
 	FOnStartSlowTask StartSlowTaskEvent;
 	FOnFinalizeSlowTask FinalizeSlowTaskEvent;
+	
+	FGuid TaskGuid;
+	FOnStartSlowTaskWithGuid StartSlowTaskWithGuidEvent;
+	FOnFinalizeSlowTaskWithGuid FinalizeSlowTaskWithGuidEvent;
 
 protected:
 	
