Fix for transition to PIE where no map has been loaded

diff --git a/Engine/Plugins/Experimental/StudioTelemetry/Source/StudioTelemetry/Private/StudioTelemetryEditor.cpp b/Engine/Plugins/Experimental/StudioTelemetry/Source/StudioTelemetry/Private/StudioTelemetryEditor.cpp
--- a/Engine/Plugins/Experimental/StudioTelemetry/Source/StudioTelemetry/Private/StudioTelemetryEditor.cpp
+++ b/Engine/Plugins/Experimental/StudioTelemetry/Source/StudioTelemetry/Private/StudioTelemetryEditor.cpp
@@ -595,42 +595,30 @@ void FStudioTelemetryEditor::Initialize()
 
 	FWorldDelegates::OnPIEReady.AddLambda([this](UGameInstance* GameInstance)
 		{
-			// PIE is now ready for user interaction
-			static bool IsFirstTimeToPIE = true;
-
-			FStudioTelemetry::Get().EndSpan(PIEStartupSpan);
-
-			// Record the time from start PIE to PIE
-			FStudioTelemetryEditor::RecordEvent_Loading(TEXT("PIE.TotalStartupTime"), PIEStartupSpan->GetDuration(), PIEStartupSpan->GetAttributes());
-			FStudioTelemetryEditor::RecordEvent_CoreSystems(TEXT("PIE.TotalStartupTime"), PIEStartupSpan->GetAttributes());
-
-			if (IsFirstTimeToPIE == true)
+			if ( PIEStartupSpan.IsValid() )
 			{
-				const double TimeInEditor = EditorLoadMapSpan->GetDuration();
-				const double TimeToStartPIE = PIEStartupSpan->GetDuration();
-				const double TimeToBootToPIE = TimeToBootEditor + TimeInEditor + TimeToStartPIE;
+				// PIE is now ready for user interaction
+				static bool IsFirstTimeToPIE = true;
 
-				// Record the absolute time from editor boot to PIE
-				FStudioTelemetryEditor::RecordEvent_Loading(TEXT("TimeToPIE"), TimeToBootToPIE, PIEStartupSpan->GetAttributes());
-				FStudioTelemetryEditor::RecordEvent_CoreSystems(TEXT("TimeToPIE"), PIEStartupSpan->GetAttributes());
+				FStudioTelemetry::Get().EndSpan(PIEStartupSpan);
 
-				IsFirstTimeToPIE = false;
-			}
-	
-			// Start PIE World Streaming span
-			/*if (GameInstance)
-			{
-				if (UWorld* World = GameInstance->GetWorld())
+				// Record the time from start PIE to PIE
+				FStudioTelemetryEditor::RecordEvent_Loading(TEXT("PIE.TotalStartupTime"), PIEStartupSpan->GetDuration(), PIEStartupSpan->GetAttributes());
+				FStudioTelemetryEditor::RecordEvent_CoreSystems(TEXT("PIE.TotalStartupTime"), PIEStartupSpan->GetAttributes());
+
+				if (IsFirstTimeToPIE == true)
 				{
-					PIEWorldStreamingSpan = FStudioTelemetry::Get().StartSpan(PIEWorldStreamingSpanName);
+					const double TimeInEditor = EditorLoadMapSpan.IsValid() ? EditorLoadMapSpan->GetDuration() : 0.0;
+					const double TimeToStartPIE = PIEStartupSpan->GetDuration();
+					const double TimeToBootToPIE = TimeToBootEditor + TimeInEditor + TimeToStartPIE;
 
-					World->OnWorldMatchStarting.AddLambda([this]()
-						{
-							FStudioTelemetry::Get().EndSpan(PIEWorldStreamingSpan);
-							FStudioTelemetryEditor::RecordEvent_Loading(TEXT("PIE.WorldStreaming"), PIEWorldStreamingSpan->GetDuration(), PIEWorldStreamingSpan->GetAttributes());
-						});
+					// Record the absolute time from editor boot to PIE
+					FStudioTelemetryEditor::RecordEvent_Loading(TEXT("TimeToPIE"), TimeToBootToPIE, PIEStartupSpan->GetAttributes());
+					FStudioTelemetryEditor::RecordEvent_CoreSystems(TEXT("TimeToPIE"), PIEStartupSpan->GetAttributes());
+
+					IsFirstTimeToPIE = false;
 				}
-			}	*/
+			}
 		});
 
 	FEditorDelegates::EndPIE.AddLambda([this](bool)
