Added Heartbeat callback running every 5 seconds to track hitches
Added basic hitch detection telemetry to StudioTelemetryEditor
Changed const FName& to const FName across API calls so they are consistent
Added RecordEvent( Category, Name, Attributes ) to API for future feature support
Renamed RecordEvent( Provider,  Name, Attributes ) to RecordEventToProvider( ) so it is more explicit and does not confuse with  RecordEvent( Category.. )
Added new attribute Config_IsDebuggerPresent to track whether a debugger is attached to the session or not.

diff --git a/Engine/Plugins/Experimental/StudioTelemetry/Source/StudioTelemetry/Private/StudioTelemetry.cpp b/Engine/Plugins/Experimental/StudioTelemetry/Source/StudioTelemetry/Private/StudioTelemetry.cpp
--- a/Engine/Plugins/Experimental/StudioTelemetry/Source/StudioTelemetry/Private/StudioTelemetry.cpp
+++ b/Engine/Plugins/Experimental/StudioTelemetry/Source/StudioTelemetry/Private/StudioTelemetry.cpp
@@ -145,7 +145,8 @@ void FStudioTelemetry::StartSession()
 		DefaultEventAttributes.Emplace(TEXT("Config_IsUnattended"), FApp::IsUnattended());
 		DefaultEventAttributes.Emplace(TEXT("Config_IsBuildMachine"), GIsBuildMachine);
 		DefaultEventAttributes.Emplace(TEXT("Config_IsRunningCommandlet"), IsRunningCommandlet());
-
+		DefaultEventAttributes.Emplace(TEXT("Config_IsDebuggerPresent"), FPlatformMisc::IsDebuggerPresent());
+		
 #if WITH_EDITOR
 		if (!FHorde::GetJobId().IsEmpty())
 		{
@@ -191,7 +192,16 @@ void FStudioTelemetry::RecordEvent(const FString& EventName, const TArray<FAnaly
 	}
 }
 
-void FStudioTelemetry::RecordEvent(const FString& ProviderName, const FString& EventName, const TArray<FAnalyticsEventAttribute>& Attributes)
+void FStudioTelemetry::RecordEvent(const FName CategoryName, const FString& EventName, const TArray<FAnalyticsEventAttribute>& Attributes)
+{
+	if (AnalyticsProvider.IsValid())
+	{
+		FScopeLock ScopeLock(&CriticalSection);
+		AnalyticsProvider->RecordEvent(EventName, Attributes);
+	}
+}
+
+void FStudioTelemetry::RecordEventToProvider(const FString& ProviderName, const FString& EventName, const TArray<FAnalyticsEventAttribute>& Attributes)
 {
 	FScopeLock ScopeLock(&CriticalSection);
 	TSharedPtr<IAnalyticsProvider> NamedProvider = GetProvider(ProviderName).Pin();
@@ -232,7 +242,7 @@ bool FStudioTelemetry::EndSpan(TSharedPtr<IAnalyticsSpan> Span, const TArray<FAn
 	return AnalyticsTracer.IsValid() ? AnalyticsTracer->EndSpan(Span, AdditionalAttributes) : false;
 }
 
-bool FStudioTelemetry::EndSpan(const FName& Name, const TArray<FAnalyticsEventAttribute>& AdditionalAttributes)
+bool FStudioTelemetry::EndSpan(const FName Name, const TArray<FAnalyticsEventAttribute>& AdditionalAttributes)
 {
 	return AnalyticsTracer.IsValid() ? AnalyticsTracer->EndSpan(Name, AdditionalAttributes) : false;
 }
diff --git a/Engine/Plugins/Experimental/StudioTelemetry/Source/StudioTelemetry/Private/StudioTelemetryEditor.cpp b/Engine/Plugins/Experimental/StudioTelemetry/Source/StudioTelemetry/Private/StudioTelemetryEditor.cpp
--- a/Engine/Plugins/Experimental/StudioTelemetry/Source/StudioTelemetry/Private/StudioTelemetryEditor.cpp
+++ b/Engine/Plugins/Experimental/StudioTelemetry/Source/StudioTelemetry/Private/StudioTelemetryEditor.cpp
@@ -394,6 +394,50 @@ void FStudioTelemetryEditor::RegisterCollectionWorkflowDelegates(FTelemetryRoute
 	});
 }
 
+extern ENGINE_API float GAverageFPS;
+
+void FStudioTelemetryEditor::HeartbeatCallback()
+{
+	static bool WasHitchingLastTime = false;
+	static uint32 HitchCount = 0;
+
+	// Record a hitch when FPS is below our threshold
+	const bool IsHitching = GAverageFPS<MinFPSForHitching;
+
+	HitchCount += IsHitching? 1:0;
+	
+	if (WasHitchingLastTime != IsHitching)
+	{
+		if (IsHitching)
+		{
+			// Start the hitch span
+			HitchingSpan = FStudioTelemetry::Get().StartSpan(HitchingSpanName);
+			HitchCount = 1;
+		}
+		else
+		{
+			TArray<FAnalyticsEventAttribute> Attributes;
+			Attributes.Emplace(FAnalyticsEventAttribute(TEXT("Hitch_Count"), HitchCount));
+			Attributes.Emplace(FAnalyticsEventAttribute(TEXT("Hitch_HitchesPerSecond"), (float)HitchCount / HitchingSpan->GetDuration()));
+			Attributes.Emplace(FAnalyticsEventAttribute(TEXT("Hitch_AverageFPS"), GAverageFPS));
+			
+			// End the hitch Span
+			FStudioTelemetry::Get().EndSpan(HitchingSpan, Attributes);
+
+			// Record the hitch event
+			FStudioTelemetry::Get().RecordEvent(TEXT("Core.Hitch"), Attributes);
+
+			Attributes.Emplace(TEXT("MapName"), EditorMapName);
+			Attributes.Emplace(TEXT("PIE_MapName"), PIEMapName);
+
+			// Record core systems events for the hitch
+			RecordEvent_CoreSystems(TEXT("Hitch"));
+		}
+	}
+
+	WasHitchingLastTime = IsHitching;
+}
+
 void FStudioTelemetryEditor::Initialize()
 {
 	SessionStartTime = FPlatformTime::Seconds();
@@ -407,9 +451,9 @@ void FStudioTelemetryEditor::Initialize()
 
 	TArray<FAnalyticsEventAttribute> Attributes;
 	Attributes.Emplace(FAnalyticsEventAttribute(TEXT("MapName"), EditorMapName));
-
 	EditorBootSpan->AddAttributes(Attributes);
 
+	// Set up Slow Task callbacks
 	ensureMsgf(GWarn, TEXT("GWarn was not valid"));
 
 	if (GWarn != nullptr)
@@ -534,6 +578,11 @@ void FStudioTelemetryEditor::Initialize()
 
 						FStudioTelemetry::Get().EndSpan(OpenAssetEditorSpan, Attributes);			
 				});
+
+				// Setup a timer for a heartbeat event.
+				FTimerDelegate Delegate;
+				Delegate.BindRaw(this, &FStudioTelemetryEditor::HeartbeatCallback);
+				GEditor->GetTimerManager()->SetTimer(TelemetryHeartbeatTimerHandle, Delegate, HeartbeatIntervalSeconds, true);
 			}
 
 			ensureMsgf(GUnrealEd, TEXT("GUnrealEd was not valid"));
@@ -549,8 +598,6 @@ void FStudioTelemetryEditor::Initialize()
 						FStudioTelemetryEditor::RecordEvent_CoreSystems(TEXT("Cooking"), Attributes);
 					});
 			}
-
-			
 		});
 
 	// Install PIE Mode callbacks
@@ -623,8 +670,14 @@ void FStudioTelemetryEditor::Initialize()
 
 	FEditorDelegates::EndPIE.AddLambda([this](bool)
 		{
-			// PIE has ended, ie. the user has pressed the Stop PIE button, and we are going back to interactive Editor mode	
-			FStudioTelemetry::Get().EndSpan(PIESpan);
+			if (PIESpan.IsValid())
+			{
+				// PIE has ended, ie. the user has pressed the Stop PIE button, and we are going back to interactive Editor mode	
+				FStudioTelemetry::Get().EndSpan(PIESpan);
+
+				FStudioTelemetryEditor::RecordEvent_Loading(TEXT("PIE.EndTime"), PIESpan->GetDuration(), PIESpan->GetAttributes());
+				FStudioTelemetryEditor::RecordEvent_CoreSystems(TEXT("PIE.EndTime"), PIESpan->GetAttributes());
+			}
 
 			TArray<FAnalyticsEventAttribute> Attributes;
 			Attributes.Emplace(TEXT("MapName"), EditorMapName);
diff --git a/Engine/Plugins/Experimental/StudioTelemetry/Source/StudioTelemetry/Private/StudioTelemetryEditor.h b/Engine/Plugins/Experimental/StudioTelemetry/Source/StudioTelemetry/Private/StudioTelemetryEditor.h
--- a/Engine/Plugins/Experimental/StudioTelemetry/Source/StudioTelemetry/Private/StudioTelemetryEditor.h
+++ b/Engine/Plugins/Experimental/StudioTelemetry/Source/StudioTelemetry/Private/StudioTelemetryEditor.h
@@ -5,11 +5,13 @@
 #if WITH_EDITOR
 
 #include "StudioTelemetry.h"
+#include "Engine/EngineTypes.h"
 
 class FTelemetryRouter;
+struct FTimerHandle;
 
 /**
- * A class that implements a variety of pre-configured Core and Editor telemetry events that can be used to evaluate the efficiecnty of the most common developer workflows
+ * A class that implements a variety of pre-configured Core and Editor telemetry events that can be used to evaluate the efficiency of the most common developer workflows
  */
 class FStudioTelemetryEditor : FNoncopyable
 {
@@ -33,7 +35,8 @@ private:
 	static void RecordEvent_Zen(const FString& Context, TArray<FAnalyticsEventAttribute> Attributes = {});
 	static void RecordEvent_VirtualAssets(const FString& Context, TArray<FAnalyticsEventAttribute> Attributes = {});
 	static void RegisterCollectionWorkflowDelegates(FTelemetryRouter& Router);
-
+	void HeartbeatCallback();
+	
 	TSharedPtr<IAnalyticsSpan> EditorSpan;
 	TSharedPtr<IAnalyticsSpan> EditorBootSpan;
 	TSharedPtr<IAnalyticsSpan> EditorInteractSpan;
@@ -44,6 +47,7 @@ private:
 	TSharedPtr<IAnalyticsSpan> PIELoadMapSpan;
 	TSharedPtr<IAnalyticsSpan> PIEWorldStreamingSpan;
 	TSharedPtr<IAnalyticsSpan> CookingSpan;
+	TSharedPtr<IAnalyticsSpan> HitchingSpan;
 	
 	const FName EditorSpanName = TEXT("Editor");
 	const FName EditorBootSpanName = TEXT("Editor.Boot");
@@ -55,11 +59,15 @@ private:
 	const FName PIELoadMapSpanName = TEXT("PIE.LoadMap");
 	const FName PIEWorldStreamingSpanName = TEXT("PIE.WorldStreaming");
 	const FName CookingSpanName = TEXT("Cooking");
+	const FName HitchingSpanName = TEXT("Hitching");
 	const FName OpenAssetEditorSpan = TEXT("Open Asset Editor");
+	const float HeartbeatIntervalSeconds = 5.0;
+	const float MinFPSForHitching = 5.0;
 
 	TMap<FGuid, TSharedPtr<IAnalyticsSpan>> TaskSpans;
 	FCriticalSection TaskSpanCriticalSection;
 
+	FTimerHandle TelemetryHeartbeatTimerHandle;
 	FString EditorMapName;
 	FString PIEMapName;
 	double SessionStartTime;
diff --git a/Engine/Plugins/Experimental/StudioTelemetry/Source/StudioTelemetry/Public/StudioTelemetry.h b/Engine/Plugins/Experimental/StudioTelemetry/Source/StudioTelemetry/Public/StudioTelemetry.h
--- a/Engine/Plugins/Experimental/StudioTelemetry/Source/StudioTelemetry/Public/StudioTelemetry.h
+++ b/Engine/Plugins/Experimental/StudioTelemetry/Source/StudioTelemetry/Public/StudioTelemetry.h
@@ -46,8 +46,11 @@ public:
 	/** Thread safe method to record an event to all registered analytics providers*/
 	STUDIOTELEMETRY_API void RecordEvent(const FString& EventName, const TArray<FAnalyticsEventAttribute>& Attributes = {});
 
+	/** Thread safe method to record an event to all registered analytics providers*/
+	STUDIOTELEMETRY_API void RecordEvent(const FName CategoryName, const FString& EventName, const TArray<FAnalyticsEventAttribute>& Attributes = {});
+
 	/** Thread safe method to record an event to the specifically named analytics provider */
-	STUDIOTELEMETRY_API void RecordEvent(const FString& ProviderName, const FString& EventName, const TArray<FAnalyticsEventAttribute>& Attributes = {});
+	STUDIOTELEMETRY_API void RecordEventToProvider(const FString& ProviderName, const FString& EventName, const TArray<FAnalyticsEventAttribute>& Attributes = {});
 
 	/** Start a new span specifying the parent*/
 	STUDIOTELEMETRY_API TSharedPtr<IAnalyticsSpan> StartSpan(const FName Name, const TArray<FAnalyticsEventAttribute>& AdditionalAttributes = {});
@@ -59,7 +62,7 @@ public:
 	STUDIOTELEMETRY_API bool EndSpan(TSharedPtr<IAnalyticsSpan> Span, const TArray<FAnalyticsEventAttribute>& AdditionalAttributes = {});
 
 	/** End an existing span by name*/
-	STUDIOTELEMETRY_API bool EndSpan(const FName& Name, const TArray<FAnalyticsEventAttribute>& AdditionalAttributes = {});
+	STUDIOTELEMETRY_API bool EndSpan(const FName Name, const TArray<FAnalyticsEventAttribute>& AdditionalAttributes = {});
 
 	/** Get an active span by name, non active spans will not be available*/
 	STUDIOTELEMETRY_API TSharedPtr<IAnalyticsSpan> GetSpan(const FName Name);
diff --git a/Engine/Source/Runtime/Analytics/Analytics/Private/AnalyticsTracer.cpp b/Engine/Source/Runtime/Analytics/Analytics/Private/AnalyticsTracer.cpp
--- a/Engine/Source/Runtime/Analytics/Analytics/Private/AnalyticsTracer.cpp
+++ b/Engine/Source/Runtime/Analytics/Analytics/Private/AnalyticsTracer.cpp
@@ -42,7 +42,7 @@ double FAnalyticsSpan::GetDuration() const
 	return Duration;
 }
 
-void FAnalyticsSpan::Start(const FName& NewSpanName, TSharedPtr<IAnalyticsSpan> NewSpanParent, const TArray<FAnalyticsEventAttribute>& AdditionalAttributes)
+void FAnalyticsSpan::Start(const FName NewSpanName, TSharedPtr<IAnalyticsSpan> NewSpanParent, const TArray<FAnalyticsEventAttribute>& AdditionalAttributes)
 {
 	// Create a new Guid for this flow, can we assume it is unique?
 	Name			= NewSpanName;
diff --git a/Engine/Source/Runtime/Analytics/Analytics/Public/AnalyticsTracer.h b/Engine/Source/Runtime/Analytics/Analytics/Public/AnalyticsTracer.h
--- a/Engine/Source/Runtime/Analytics/Analytics/Public/AnalyticsTracer.h
+++ b/Engine/Source/Runtime/Analytics/Analytics/Public/AnalyticsTracer.h
@@ -18,7 +18,7 @@ public:
 
 	// Public IAnalyticsSpan implementation
 	virtual void SetProvider(TSharedPtr<IAnalyticsProvider> AnalyticsProvider) override;
-	virtual void Start(const FName& Name, TSharedPtr<IAnalyticsSpan> ParentSpan, const TArray<FAnalyticsEventAttribute>& AdditionalAttributes = {});
+	virtual void Start(const FName Name, TSharedPtr<IAnalyticsSpan> ParentSpan, const TArray<FAnalyticsEventAttribute>& AdditionalAttributes = {});
 	virtual void End(const TArray<FAnalyticsEventAttribute>& AdditionalAttributes = {}) override;
 	virtual void AddAttributes(const TArray<FAnalyticsEventAttribute>& AdditionalAttributes) override;
 	virtual void RecordEvent(const FString& EventName, const TArray<FAnalyticsEventAttribute>& AdditionalAttributes = {}) override;
diff --git a/Engine/Source/Runtime/Analytics/Analytics/Public/Interfaces/IAnalyticsTracer.h b/Engine/Source/Runtime/Analytics/Analytics/Public/Interfaces/IAnalyticsTracer.h
--- a/Engine/Source/Runtime/Analytics/Analytics/Public/Interfaces/IAnalyticsTracer.h
+++ b/Engine/Source/Runtime/Analytics/Analytics/Public/Interfaces/IAnalyticsTracer.h
@@ -28,7 +28,7 @@ public:
 	virtual void SetProvider(TSharedPtr<IAnalyticsProvider> AnalyticsProvider) = 0;
 
 	/** Start this span */
-	virtual void Start(const FName& Name, TSharedPtr<IAnalyticsSpan> ParentSpan, const TArray<FAnalyticsEventAttribute>& AdditionalAttributes = {}) = 0;
+	virtual void Start(const FName Name, TSharedPtr<IAnalyticsSpan> ParentSpan, const TArray<FAnalyticsEventAttribute>& AdditionalAttributes = {}) = 0;
 
 	/** End this span */
 	virtual void End(const TArray<FAnalyticsEventAttribute>& AdditionalAttributes = {}) = 0;
@@ -104,7 +104,7 @@ public:
 	};
 
 	/** End an existing span by name*/
-	bool EndSpan(const FName& Name, const TArray<FAnalyticsEventAttribute>& AdditionalAttributes = {})
+	bool EndSpan(const FName Name, const TArray<FAnalyticsEventAttribute>& AdditionalAttributes = {})
 	{
 		return EndSpan(GetSpan(Name), AdditionalAttributes);
 	}
