Fix to hitch span logic so that spans closed by parent spans whilst hithcing creates a new hitch span. This means we now get hitch detection immediately after PIE has started if it was already hithcing before it ended.
Added IsActive to FAnalyticsSpan API
Reduced Hitch Detector frequency from 5sec to 1sec
Added Asset Registry Activity span

diff --git a/Engine/Plugins/Experimental/StudioTelemetry/Source/StudioTelemetry/Private/StudioTelemetryEditor.cpp b/Engine/Plugins/Experimental/StudioTelemetry/Source/StudioTelemetry/Private/StudioTelemetryEditor.cpp
--- a/Engine/Plugins/Experimental/StudioTelemetry/Source/StudioTelemetry/Private/StudioTelemetryEditor.cpp
+++ b/Engine/Plugins/Experimental/StudioTelemetry/Source/StudioTelemetry/Private/StudioTelemetryEditor.cpp
@@ -13,6 +13,7 @@
 #include "AssetRegistry/AssetRegistryTelemetry.h"
 #include "Editor.h"
 #include "Editor/UnrealEdEngine.h"
+#include "Engine/AssetManager.h"
 #include "UnrealEdGlobals.h"
 #include "CookOnTheSide/CookOnTheFlyServer.h"
 #include "Subsystems/AssetEditorSubsystem.h"
@@ -398,44 +399,63 @@ extern ENGINE_API float GAverageFPS;
 
 void FStudioTelemetryEditor::HeartbeatCallback()
 {
-	static bool WasHitchingLastTime = false;
 	static uint32 HitchCount = 0;
 
-	// Record a hitch when FPS is below our threshold
+	// Hitching is when FPS is below our threshold
 	const bool IsHitching = GAverageFPS<MinFPSForHitching;
 
-	HitchCount += IsHitching? 1:0;
+	HitchCount += IsHitching ? 1:0;
 	
-	if (WasHitchingLastTime != IsHitching)
+	if ( IsHitching == false && HitchingSpan.IsValid()==true )
 	{
-		if (IsHitching)
-		{
-			// Start the hitch span
-			HitchingSpan = FStudioTelemetry::Get().StartSpan(HitchingSpanName);
-			HitchCount = 1;
-		}
-		else
-		{
-			TArray<FAnalyticsEventAttribute> Attributes;
-			Attributes.Emplace(FAnalyticsEventAttribute(TEXT("Hitch_Count"), HitchCount));
-			Attributes.Emplace(FAnalyticsEventAttribute(TEXT("Hitch_HitchesPerSecond"), (float)HitchCount / HitchingSpan->GetDuration()));
-			Attributes.Emplace(FAnalyticsEventAttribute(TEXT("Hitch_AverageFPS"), GAverageFPS));
-			
-			// End the hitch Span
-			FStudioTelemetry::Get().EndSpan(HitchingSpan, Attributes);
+		// No longer hitching and we have started a hitch span
+		TArray<FAnalyticsEventAttribute> Attributes;
+		Attributes.Emplace(FAnalyticsEventAttribute(TEXT("Hitch_Count"), HitchCount));
+		Attributes.Emplace(FAnalyticsEventAttribute(TEXT("Hitch_HitchesPerSecond"), (float)HitchCount / HitchingSpan->GetDuration()));
+		Attributes.Emplace(FAnalyticsEventAttribute(TEXT("Hitch_AverageFPS"), GAverageFPS));
 
-			// Record the hitch event
-			FStudioTelemetry::Get().RecordEvent(TEXT("Core.Hitch"), Attributes);
+		// End the hitch Span
+		FStudioTelemetry::Get().EndSpan(HitchingSpan, Attributes);
 
-			Attributes.Emplace(TEXT("MapName"), EditorMapName);
-			Attributes.Emplace(TEXT("PIE_MapName"), PIEMapName);
+		// Record the hitch event
+		FStudioTelemetry::Get().RecordEvent(TEXT("Core.Hitch"), Attributes);
 
-			// Record core systems events for the hitch
-			RecordEvent_CoreSystems(TEXT("Hitch"));
-		}
+		Attributes.Emplace(TEXT("MapName"), EditorMapName);
+		Attributes.Emplace(TEXT("PIE_MapName"), PIEMapName);
+
+		// Record core systems events for the hitch
+		RecordEvent_CoreSystems(TEXT("Hitch"));
+
+		// No longer need the hitch span for now so reset it
+		HitchingSpan.Reset();
 	}
+	else if ( IsHitching == true && HitchingSpan.IsValid()==false )
+	{
+		// We are hitching and we have not started a hitch span
+		HitchingSpan = FStudioTelemetry::Get().StartSpan(HitchingSpanName);
+		HitchCount = 1;
+	}	
+
+	// Monitor Asset Registry Scan
+	FAssetRegistryModule& AssetRegistryModule = FModuleManager::LoadModuleChecked<FAssetRegistryModule>(TEXT("AssetRegistry"));
+	const bool IsScaningAssets = AssetRegistryModule.Get().IsLoadingAssets() || AssetRegistryModule.Get().IsSearchAllAssets();
+
+	if (IsScaningAssets == false && AssetRegistryScanSpan.IsValid() == true)
+	{
+		// End the span
+		TArray<FAnalyticsEventAttribute> Attributes;
+		Attributes.Emplace(TEXT("MapName"), EditorMapName);
+
+		FStudioTelemetry::Get().EndSpan(AssetRegistryScanSpan, Attributes);
 
-	WasHitchingLastTime = IsHitching;
+		// No longer need the span for now so reset it
+		AssetRegistryScanSpan.Reset();
+	}
+	else if (IsScaningAssets == true && AssetRegistryScanSpan.IsValid() == false)
+	{
+		// Start the span
+		AssetRegistryScanSpan = FStudioTelemetry::Get().StartSpan(AssetRegistryScanSpanName);
+	}
 }
 
 void FStudioTelemetryEditor::Initialize()
@@ -443,7 +463,7 @@ void FStudioTelemetryEditor::Initialize()
 	SessionStartTime = FPlatformTime::Seconds();
 
 	// Install Editor Mode callbacks
-	// 
+	
 	// Start Editor and Editor Boot span. Note : this will only start when the plugin is loaded and as such will miss any activity that runs beforehand
 	EditorSpan = FStudioTelemetry::Get().StartSpan(EditorSpanName);
 	EditorBootSpan = FStudioTelemetry::Get().StartSpan(EditorBootSpanName);
diff --git a/Engine/Plugins/Experimental/StudioTelemetry/Source/StudioTelemetry/Private/StudioTelemetryEditor.h b/Engine/Plugins/Experimental/StudioTelemetry/Source/StudioTelemetry/Private/StudioTelemetryEditor.h
--- a/Engine/Plugins/Experimental/StudioTelemetry/Source/StudioTelemetry/Private/StudioTelemetryEditor.h
+++ b/Engine/Plugins/Experimental/StudioTelemetry/Source/StudioTelemetry/Private/StudioTelemetryEditor.h
@@ -48,6 +48,7 @@ private:
 	TSharedPtr<IAnalyticsSpan> PIEWorldStreamingSpan;
 	TSharedPtr<IAnalyticsSpan> CookingSpan;
 	TSharedPtr<IAnalyticsSpan> HitchingSpan;
+	TSharedPtr<IAnalyticsSpan> AssetRegistryScanSpan;
 	
 	const FName EditorSpanName = TEXT("Editor");
 	const FName EditorBootSpanName = TEXT("Editor.Boot");
@@ -61,7 +62,8 @@ private:
 	const FName CookingSpanName = TEXT("Cooking");
 	const FName HitchingSpanName = TEXT("Hitching");
 	const FName OpenAssetEditorSpan = TEXT("Open Asset Editor");
-	const float HeartbeatIntervalSeconds = 5.0;
+	const FName AssetRegistryScanSpanName = TEXT("Asset Registry Scan");
+	const float HeartbeatIntervalSeconds = 0.5;
 	const float MinFPSForHitching = 5.0;
 
 	TMap<FGuid, TSharedPtr<IAnalyticsSpan>> TaskSpans;
diff --git a/Engine/Source/Runtime/Analytics/Analytics/Private/AnalyticsTracer.cpp b/Engine/Source/Runtime/Analytics/Analytics/Private/AnalyticsTracer.cpp
--- a/Engine/Source/Runtime/Analytics/Analytics/Private/AnalyticsTracer.cpp
+++ b/Engine/Source/Runtime/Analytics/Analytics/Private/AnalyticsTracer.cpp
@@ -64,6 +64,11 @@ void FAnalyticsSpan::AddChildSpan(TSharedPtr<IAnalyticsSpan> ChildSpan)
 	ChildSpans.Add(ChildSpan);
 }
 
+bool FAnalyticsSpan::GetIsActive() const
+{
+	return IsActive;
+}
+
 void FAnalyticsSpan::End(const TArray<FAnalyticsEventAttribute>& AdditionalAttributes)
 {
 	// Only End the span once
diff --git a/Engine/Source/Runtime/Analytics/Analytics/Public/AnalyticsTracer.h b/Engine/Source/Runtime/Analytics/Analytics/Public/AnalyticsTracer.h
--- a/Engine/Source/Runtime/Analytics/Analytics/Public/AnalyticsTracer.h
+++ b/Engine/Source/Runtime/Analytics/Analytics/Public/AnalyticsTracer.h
@@ -28,6 +28,7 @@ public:
 	virtual TSharedPtr<IAnalyticsSpan> GetParentSpan() const override;
 	virtual double GetDuration() const override;
 	virtual void AddChildSpan(TSharedPtr<IAnalyticsSpan> ChildSpan) override;
+	virtual bool GetIsActive() const override;
 
 private:
 
diff --git a/Engine/Source/Runtime/Analytics/Analytics/Public/Interfaces/IAnalyticsTracer.h b/Engine/Source/Runtime/Analytics/Analytics/Public/Interfaces/IAnalyticsTracer.h
--- a/Engine/Source/Runtime/Analytics/Analytics/Public/Interfaces/IAnalyticsTracer.h
+++ b/Engine/Source/Runtime/Analytics/Analytics/Public/Interfaces/IAnalyticsTracer.h
@@ -51,6 +51,9 @@ public:
 	/** Get the span duration in seconds */
 	virtual double GetDuration() const = 0;
 
+	// Is the span active?
+	virtual bool GetIsActive() const = 0;
+
 	/** Get the parent span */
 	virtual TSharedPtr<IAnalyticsSpan> GetParentSpan() const = 0;
 
