Added PIE_TransitionCount to the Telemetry

diff --git a/Engine/Plugins/Experimental/StudioTelemetry/Source/StudioTelemetry/Private/StudioTelemetryEditor.cpp b/Engine/Plugins/Experimental/StudioTelemetry/Source/StudioTelemetry/Private/StudioTelemetryEditor.cpp
--- a/Engine/Plugins/Experimental/StudioTelemetry/Source/StudioTelemetry/Private/StudioTelemetryEditor.cpp
+++ b/Engine/Plugins/Experimental/StudioTelemetry/Source/StudioTelemetry/Private/StudioTelemetryEditor.cpp
@@ -678,15 +678,21 @@ void FStudioTelemetryEditor::Initialize()
 				if ( PIEStartupSpan.IsValid() )
 				{
 					// PIE is now ready for user interaction
-					static bool IsFirstTimeToPIE = true;
-
+					
+					// Keep track of the PIE transition counts
+					TArray<FAnalyticsEventAttribute> Attributes;
+					Attributes.Emplace(TEXT("PIE_TransitionCount"), PIETransitionCount );
+					
+					PIESpan->AddAttributes(Attributes);
+					PIEStartupSpan->AddAttributes(Attributes);
+						
 					FStudioTelemetry::Get().EndSpan(PIEStartupSpan);
 
 					// Record the time from start PIE to PIE
 					FStudioTelemetryEditor::RecordEvent_Loading(TEXT("PIE.TotalStartupTime"), PIEStartupSpan->GetDuration(), PIEStartupSpan->GetAttributes());
 					FStudioTelemetryEditor::RecordEvent_CoreSystems(TEXT("PIE.TotalStartupTime"), PIEStartupSpan->GetAttributes());
 
-					if (IsFirstTimeToPIE == true)
+					if (PIETransitionCount == 0)
 					{
 						const double TimeInEditor = EditorLoadMapSpan.IsValid() ? EditorLoadMapSpan->GetDuration() : 0.0;
 						const double TimeToStartPIE = PIEStartupSpan->GetDuration();
@@ -695,9 +701,9 @@ void FStudioTelemetryEditor::Initialize()
 						// Record the absolute time from editor boot to PIE
 						FStudioTelemetryEditor::RecordEvent_Loading(TEXT("TimeToPIE"), TimeToBootToPIE, PIEStartupSpan->GetAttributes());
 						FStudioTelemetryEditor::RecordEvent_CoreSystems(TEXT("TimeToPIE"), PIEStartupSpan->GetAttributes());
-
-						IsFirstTimeToPIE = false;
 					}
+
+					PIETransitionCount++;
 				}
 
 				PIEInteractSpan = FStudioTelemetry::Get().StartSpan(PIEInteractSpanName, PIESpan);
diff --git a/Engine/Plugins/Experimental/StudioTelemetry/Source/StudioTelemetry/Private/StudioTelemetryEditor.h b/Engine/Plugins/Experimental/StudioTelemetry/Source/StudioTelemetry/Private/StudioTelemetryEditor.h
--- a/Engine/Plugins/Experimental/StudioTelemetry/Source/StudioTelemetry/Private/StudioTelemetryEditor.h
+++ b/Engine/Plugins/Experimental/StudioTelemetry/Source/StudioTelemetry/Private/StudioTelemetryEditor.h
@@ -76,6 +76,7 @@ private:
 	FTimerHandle TelemetryHeartbeatTimerHandle;
 	FString EditorMapName;
 	FString PIEMapName;
+	uint32 PIETransitionCount = 0;
 	double SessionStartTime;
 	double AssetOpenStartTime;
 	double TimeToBootEditor;
