Added a work around for overlapping Unreal Regions with the same name
Send SlowTask as a the span name for slow tasks instead of the dynamically created string. This makes filitering to SpanName almost impossible in the telemetry
Reduced the high frequency generation of Span events as this made Tableau dashbaords unusable
Added a new hitch detector that samples AverageFPS over a high frequency ( every 0.5 sec ) but records hitch spans at a much lower rate ( every 5 seconds )
Increased the schema version for spans to 2 so we can filter out the old high frequency data

diff --git a/Engine/Plugins/Experimental/StudioTelemetry/Source/StudioTelemetry/Private/StudioTelemetryEditor.cpp b/Engine/Plugins/Experimental/StudioTelemetry/Source/StudioTelemetry/Private/StudioTelemetryEditor.cpp
--- a/Engine/Plugins/Experimental/StudioTelemetry/Source/StudioTelemetry/Private/StudioTelemetryEditor.cpp
+++ b/Engine/Plugins/Experimental/StudioTelemetry/Source/StudioTelemetry/Private/StudioTelemetryEditor.cpp
@@ -396,44 +396,56 @@ void FStudioTelemetryEditor::RegisterCollectionWorkflowDelegates(FTelemetryRoute
 
 extern ENGINE_API float GAverageFPS;
 
+void FStudioTelemetryEditor::HitchSamplerCallback()
+{
+	// Sample a rolling average of FPS 
+	HitchAvergageFPS = ( HitchAvergageFPS * HitchSampleCount + GAverageFPS ) / (double)(HitchSampleCount+1);
+	HitchSampleCount++;
+}
+
 void FStudioTelemetryEditor::HeartbeatCallback()
 {
-	static uint32 HitchCount = 0;
+	if (HitchSampleCount>0)
+	{
+		// Hitching is when FPS is below our threshold
+		const bool IsHitching = HitchAvergageFPS < MinFPSForHitching;
+		static uint32 HitchCount = 0;
+		HitchCount += IsHitching ? 1 : 0;
 
-	// Hitching is when FPS is below our threshold
-	const bool IsHitching = GAverageFPS<MinFPSForHitching;
+		if (IsHitching == false && HitchingSpan.IsValid() == true)
+		{
+			// No longer hitching and we have started a hitch span
+			TArray<FAnalyticsEventAttribute> Attributes;
+			Attributes.Emplace(FAnalyticsEventAttribute(TEXT("Hitch_Count"), HitchCount));
+			Attributes.Emplace(FAnalyticsEventAttribute(TEXT("Hitch_HitchesPerSecond"), (float)HitchCount / HitchingSpan->GetDuration()));
+			Attributes.Emplace(FAnalyticsEventAttribute(TEXT("Hitch_AverageFPS"), HitchAvergageFPS));
 
-	HitchCount += IsHitching ? 1:0;
-	
-	if ( IsHitching == false && HitchingSpan.IsValid()==true )
-	{
-		// No longer hitching and we have started a hitch span
-		TArray<FAnalyticsEventAttribute> Attributes;
-		Attributes.Emplace(FAnalyticsEventAttribute(TEXT("Hitch_Count"), HitchCount));
-		Attributes.Emplace(FAnalyticsEventAttribute(TEXT("Hitch_HitchesPerSecond"), (float)HitchCount / HitchingSpan->GetDuration()));
-		Attributes.Emplace(FAnalyticsEventAttribute(TEXT("Hitch_AverageFPS"), GAverageFPS));
+			// End the hitch Span
+			FStudioTelemetry::Get().EndSpan(HitchingSpan, Attributes);
 
-		// End the hitch Span
-		FStudioTelemetry::Get().EndSpan(HitchingSpan, Attributes);
+			// Record the hitch event
+			FStudioTelemetry::Get().RecordEvent(TEXT("Core.Hitch"), Attributes);
 
-		// Record the hitch event
-		FStudioTelemetry::Get().RecordEvent(TEXT("Core.Hitch"), Attributes);
+			Attributes.Emplace(TEXT("MapName"), EditorMapName);
+			Attributes.Emplace(TEXT("PIE_MapName"), PIEMapName);
 
-		Attributes.Emplace(TEXT("MapName"), EditorMapName);
-		Attributes.Emplace(TEXT("PIE_MapName"), PIEMapName);
+			// Record core systems events for the hitch
+			RecordEvent_CoreSystems(TEXT("Hitch"));
 
-		// Record core systems events for the hitch
-		RecordEvent_CoreSystems(TEXT("Hitch"));
+			// No longer need the hitch span for now so reset it
+			HitchingSpan.Reset();
+		}
+		else if (IsHitching == true && HitchingSpan.IsValid() == false)
+		{
+			// We are hitching and we have not started a hitch span
+			HitchingSpan = FStudioTelemetry::Get().StartSpan(HitchingSpanName);
+			HitchCount = 1;
+		}
 
-		// No longer need the hitch span for now so reset it
-		HitchingSpan.Reset();
+		// Reset the hitch sampler
+		HitchSampleCount = 0;
+		HitchAvergageFPS = 0;
 	}
-	else if ( IsHitching == true && HitchingSpan.IsValid()==false )
-	{
-		// We are hitching and we have not started a hitch span
-		HitchingSpan = FStudioTelemetry::Get().StartSpan(HitchingSpanName);
-		HitchCount = 1;
-	}	
 
 	// Monitor Asset Registry Scan
 	FAssetRegistryModule& AssetRegistryModule = FModuleManager::LoadModuleChecked<FAssetRegistryModule>(TEXT("AssetRegistry"));
@@ -493,8 +505,7 @@ void FStudioTelemetryEditor::Initialize()
 					Attributes.Emplace(TEXT("TaskName"), TaskName.ToString());
 
 					// Create and start a new slow task span
-					const FName SpanName(TEXT("ST ") + TaskName.ToString());
-					TSharedPtr<IAnalyticsSpan> SlowTaskSpan = FStudioTelemetry::Get().StartSpan(SpanName, Attributes);
+					TSharedPtr<IAnalyticsSpan> SlowTaskSpan = FStudioTelemetry::Get().StartSpan(*TaskName.ToString(), Attributes);
 
 					// Store this SlowTask span so we can find it when it finishes
 					TaskSpans.Add(TaskGuid, SlowTaskSpan);
@@ -598,10 +609,18 @@ void FStudioTelemetryEditor::Initialize()
 						FStudioTelemetry::Get().EndSpan(OpenAssetEditorSpan, Attributes);			
 				});
 
-				// Setup a timer for a heartbeat event.
-				FTimerDelegate Delegate;
-				Delegate.BindRaw(this, &FStudioTelemetryEditor::HeartbeatCallback);
-				GEditor->GetTimerManager()->SetTimer(TelemetryHeartbeatTimerHandle, Delegate, HeartbeatIntervalSeconds, true);
+				// Setup a timer for a Heartbeat callback
+				FTimerDelegate HeartbeatDelegate;
+				HeartbeatDelegate.BindRaw(this, &FStudioTelemetryEditor::HeartbeatCallback);
+				GEditor->GetTimerManager()->SetTimer(TelemetryHeartbeatTimerHandle, HeartbeatDelegate, HeartbeatIntervalSeconds, true);
+
+				// Setup the timer for the Hitch Detector callback
+				if (IsRunningCommandlet() == false)
+				{
+					FTimerDelegate HitchSamplerDelegate;
+					HitchSamplerDelegate.BindRaw(this, &FStudioTelemetryEditor::HitchSamplerCallback);
+					GEditor->GetTimerManager()->SetTimer(TelemetryHitchSamplerTimerHandle, HitchSamplerDelegate, HitchSamplerIntervalSeconds, true);
+				}
 			}
 
 			ensureMsgf(GUnrealEd, TEXT("GUnrealEd was not valid"));
diff --git a/Engine/Plugins/Experimental/StudioTelemetry/Source/StudioTelemetry/Private/StudioTelemetryEditor.h b/Engine/Plugins/Experimental/StudioTelemetry/Source/StudioTelemetry/Private/StudioTelemetryEditor.h
--- a/Engine/Plugins/Experimental/StudioTelemetry/Source/StudioTelemetry/Private/StudioTelemetryEditor.h
+++ b/Engine/Plugins/Experimental/StudioTelemetry/Source/StudioTelemetry/Private/StudioTelemetryEditor.h
@@ -35,6 +35,7 @@ private:
 	static void RecordEvent_Zen(const FString& Context, TArray<FAnalyticsEventAttribute> Attributes = {});
 	static void RecordEvent_VirtualAssets(const FString& Context, TArray<FAnalyticsEventAttribute> Attributes = {});
 	static void RegisterCollectionWorkflowDelegates(FTelemetryRouter& Router);
+	void HitchSamplerCallback();
 	void HeartbeatCallback();
 	
 	TSharedPtr<IAnalyticsSpan> EditorSpan;
@@ -67,19 +68,23 @@ private:
 	const FName HitchingSpanName = TEXT("Hitching");
 	const FName OpenAssetEditorSpan = TEXT("Open Asset Editor");
 	const FName AssetRegistryScanSpanName = TEXT("Asset Registry Scan");
-	const float HeartbeatIntervalSeconds = 0.5;
+	const float HeartbeatIntervalSeconds = 5.0;
+	const float HitchSamplerIntervalSeconds = 0.5;
 	const float MinFPSForHitching = 5.0;
 
 	TMap<FGuid, TSharedPtr<IAnalyticsSpan>> TaskSpans;
 	FCriticalSection TaskSpanCriticalSection;
 
 	FTimerHandle TelemetryHeartbeatTimerHandle;
+	FTimerHandle TelemetryHitchSamplerTimerHandle;
 	FString EditorMapName;
 	FString PIEMapName;
 	uint32 PIETransitionCount = 0;
 	double SessionStartTime;
 	double AssetOpenStartTime;
 	double TimeToBootEditor;
+	double HitchAvergageFPS = 0;
+	uint32 HitchSampleCount = 0;
 };
 
 #endif // WITH_EDITOR
diff --git a/Engine/Source/Runtime/Analytics/Analytics/Private/AnalyticsTracer.cpp b/Engine/Source/Runtime/Analytics/Analytics/Private/AnalyticsTracer.cpp
--- a/Engine/Source/Runtime/Analytics/Analytics/Private/AnalyticsTracer.cpp
+++ b/Engine/Source/Runtime/Analytics/Analytics/Private/AnalyticsTracer.cpp
@@ -57,8 +57,6 @@ void FAnalyticsSpan::Start(const TArray<FAnalyticsEventAttribute>& AdditionalAtt
 	EndTime			= FDateTime::UtcNow();
 	Duration		= 0;	
 	IsActive		= true;
-
-	TRACE_BEGIN_REGION(*Name.ToString());
 }
 
 bool FAnalyticsSpan::GetIsActive() const
@@ -78,12 +76,10 @@ void FAnalyticsSpan::End(const TArray<FAnalyticsEventAttribute>& AdditionalAttri
 	EndTime = FDateTime::UtcNow();
 	Duration = (EndTime - StartTime).GetTotalSeconds();
 		
-	TRACE_END_REGION(*Name.ToString());
-
 	// Add attributes and the the additional attributes to the current span attributes, these will get passed down to the child spans
 	AddAttributes(AdditionalAttributes);
 
-	const uint32 SpanSchemaVersion = 1;
+	const uint32 SpanSchemaVersion = 2;
 	const FString SpanEventName = TEXT("Span");
 
 	TArray<FAnalyticsEventAttribute> EventAttributes = Attributes;
@@ -147,7 +143,7 @@ void FAnalyticsTracer::SetProvider(TSharedPtr<IAnalyticsProvider> InProvider)
 
 TSharedPtr<IAnalyticsSpan> FAnalyticsTracer::GetCurrentSpan() const
 {
-	return ActiveSpanStack.Num()? ActiveSpanStack.Top() : TSharedPtr<IAnalyticsSpan>();
+	return ActiveSpanStack.Num() ? ActiveSpanStack.Top() : TSharedPtr<IAnalyticsSpan>();
 }
 
 void FAnalyticsTracer::StartSession()
@@ -157,14 +153,14 @@ void FAnalyticsTracer::StartSession()
 
 void FAnalyticsTracer::EndSession()
 {
-	FScopeLock ScopeLock(&CriticalSection);	
+	FScopeLock ScopeLock(&CriticalSection);
 
 	EndSpan(SessionSpan);
 	SessionSpan.Reset();
 
 	// Stop any active spans, go from stack bottom first so parent spans will end their children
 	while (ActiveSpanStack.Num())
-	{	
+	{
 		EndSpan(ActiveSpanStack[0]);
 	}
 
@@ -192,18 +188,57 @@ TSharedPtr<IAnalyticsSpan> FAnalyticsTracer::StartSpan(const FName NewSpanName,
 			SpanHeirarchy.Emplace(ParentSpan->GetId(), { NewSpan });
 		}
 	}
-	
+
 
 	return StartSpanInternal(NewSpan, AdditionalAttributes) ? NewSpan : TSharedPtr<IAnalyticsSpan>();
 }
 
+void FAnalyticsTracer::BeginRegion(TSharedPtr<IAnalyticsSpan> Span)
+{
+	// This function is a temporary work around as UnrealInsights does not handle overalapping region with the same name.
+	FName RegionName = Span->GetName();
+
+	uint32 NameCounter = 0;
+
+	while (RegionNames.Find(RegionName) != nullptr)
+	{
+		// Generate a unique region name for this span
+		FNameBuilder NameBuilder(Span->GetName());
+		NameBuilder.Append(TEXT("%d"), NameCounter++);
+		RegionName = FName(NameBuilder);
+	}
+
+	// Add the region name to the list by span ID
+	RegionNames.Emplace(RegionName, Span->GetId());
+
+	TRACE_BEGIN_REGION(*RegionName.ToString());
+}
+
+void FAnalyticsTracer::EndRegion(TSharedPtr<IAnalyticsSpan> Span)
+{
+	// Find the region, end it and remove from the region names
+	for (TMap<FName, FGuid>::TConstIterator it(RegionNames); it; ++it)
+	{
+		// Slow match by ID on removal, fast match on creation
+		if ((*it).Value == Span->GetId())
+		{
+			const FName& RegionName = (*it).Key;
+			TRACE_END_REGION(*RegionName.ToString());
+			RegionNames.Remove(RegionName);
+			return;
+		}
+	}
+}
+
 bool FAnalyticsTracer::StartSpanInternal(TSharedPtr<IAnalyticsSpan> Span, const TArray<FAnalyticsEventAttribute>& AdditionalAttributes)
 {
-	TSharedPtr<IAnalyticsSpan> LastAdddedActiveSpan = ActiveSpanStack.Num()? ActiveSpanStack.Top(): TSharedPtr<IAnalyticsSpan>();
-	Span->SetStackDepth(LastAdddedActiveSpan.IsValid()? LastAdddedActiveSpan->GetStackDepth()+1 : 0);
+	TSharedPtr<IAnalyticsSpan> LastAdddedActiveSpan = ActiveSpanStack.Num() ? ActiveSpanStack.Top() : TSharedPtr<IAnalyticsSpan>();
+	Span->SetStackDepth(LastAdddedActiveSpan.IsValid() ? LastAdddedActiveSpan->GetStackDepth() + 1 : 0);
 	Span->SetProvider(AnalyticsProvider);
 	Span->Start(AdditionalAttributes);
 
+	BeginRegion(Span);
+
 	// Add span to active spans list
 	ActiveSpanStack.Emplace(Span);
 
@@ -222,6 +257,8 @@ bool FAnalyticsTracer::EndSpanInternal(TSharedPtr<IAnalyticsSpan> Span, const TA
 	{
 		Span->End(AdditionalAttributes);
 
+		EndRegion(Span);
+
 		ActiveSpanStack.Remove(Span);
 
 		// End any children of this span
diff --git a/Engine/Source/Runtime/Analytics/Analytics/Public/AnalyticsTracer.h b/Engine/Source/Runtime/Analytics/Analytics/Public/AnalyticsTracer.h
--- a/Engine/Source/Runtime/Analytics/Analytics/Public/AnalyticsTracer.h
+++ b/Engine/Source/Runtime/Analytics/Analytics/Public/AnalyticsTracer.h
@@ -69,11 +69,14 @@ private:
 	bool StartSpanInternal(TSharedPtr<IAnalyticsSpan> Span, const TArray<FAnalyticsEventAttribute>& AdditionalAttributes);
 	bool EndSpanInternal(TSharedPtr<IAnalyticsSpan> Span, const TArray<FAnalyticsEventAttribute>& AdditionalAttributes);
 	TSharedPtr<IAnalyticsSpan> GetSpanInternal(const FName Name);
+	void BeginRegion(TSharedPtr<IAnalyticsSpan> Span);
+	void EndRegion(TSharedPtr<IAnalyticsSpan> Span);
 
 	TSharedPtr<IAnalyticsProvider>					AnalyticsProvider;	// The Analytics provider we will send our span events to
 	TSharedPtr<IAnalyticsSpan>						SessionSpan;		// The root span, this will always be present in an active session
 	TArray<TSharedPtr<IAnalyticsSpan>>				ActiveSpanStack;	// Stack of active spans as WeakPtrs
 	TMap<FGuid, TArray<TWeakPtr<IAnalyticsSpan>>>	SpanHeirarchy;		// List of child spans per parent span
+	TMap<FName, FGuid>								RegionNames;		// Temprorary fix for duplicate span names in regions	
 	FCriticalSection								CriticalSection;	
 };
 
