Submitted on behalf of matt.breindel

Minor change to fix up an issue from CL 31394907 at Mark's request. The SlowTask span being created was still being named *TaskName but was meant to be changed to "SlowTask".

diff --git a/Engine/Plugins/Experimental/StudioTelemetry/Source/StudioTelemetry/Private/StudioTelemetryEditor.cpp b/Engine/Plugins/Experimental/StudioTelemetry/Source/StudioTelemetry/Private/StudioTelemetryEditor.cpp
--- a/Engine/Plugins/Experimental/StudioTelemetry/Source/StudioTelemetry/Private/StudioTelemetryEditor.cpp
+++ b/Engine/Plugins/Experimental/StudioTelemetry/Source/StudioTelemetry/Private/StudioTelemetryEditor.cpp
@@ -505,7 +505,7 @@ void FStudioTelemetryEditor::Initialize()
 					Attributes.Emplace(TEXT("TaskName"), TaskName.ToString());
 
 					// Create and start a new slow task span
-					TSharedPtr<IAnalyticsSpan> SlowTaskSpan = FStudioTelemetry::Get().StartSpan(*TaskName.ToString(), Attributes);
+					TSharedPtr<IAnalyticsSpan> SlowTaskSpan = FStudioTelemetry::Get().StartSpan(TEXT("SlowTask"), Attributes);
 
 					// Store this SlowTask span so we can find it when it finishes
 					TaskSpans.Add(TaskGuid, SlowTaskSpan);
