Studio Telemetry Update
Enable hitch detection in Studio Telemetry only when the application is in the foreground.
Increased the Hitch detection frequency as this was not the cause of the Telemetry dashboard inefficiency.
Added explicit TaskName Insights Region

diff --git a/Engine/Plugins/Experimental/StudioTelemetry/Source/StudioTelemetry/Private/StudioTelemetryEditor.cpp b/Engine/Plugins/Experimental/StudioTelemetry/Source/StudioTelemetry/Private/StudioTelemetryEditor.cpp
--- a/Engine/Plugins/Experimental/StudioTelemetry/Source/StudioTelemetry/Private/StudioTelemetryEditor.cpp
+++ b/Engine/Plugins/Experimental/StudioTelemetry/Source/StudioTelemetry/Private/StudioTelemetryEditor.cpp
@@ -398,11 +398,15 @@ extern ENGINE_API float GAverageFPS;
 
 void FStudioTelemetryEditor::HitchSamplerCallback()
 {
-	// Sample a rolling average of FPS 
-	HitchAvergageFPS = ( HitchAvergageFPS * HitchSampleCount + GAverageFPS ) / (double)(HitchSampleCount+1);
-	HitchSampleCount++;
+	// Only sample framerate when we have focus
+	if (FApp::HasFocus())
+	{
+		// Sample a rolling average of FPS 
+		HitchAvergageFPS = (HitchAvergageFPS * HitchSampleCount + GAverageFPS) / (double)(HitchSampleCount + 1);
+		HitchSampleCount++;
+	}
 }
-
+	
 void FStudioTelemetryEditor::HeartbeatCallback()
 {
 	if (HitchSampleCount>0)
@@ -510,11 +514,15 @@ void FStudioTelemetryEditor::Initialize()
 					// Store this SlowTask span so we can find it when it finishes
 					TaskSpans.Add(TaskGuid, SlowTaskSpan);
 				}
+
+				TRACE_BEGIN_REGION(*TaskName.ToString());
 			});
 
 		// End the SlowTask span
-		GWarn->OnFinalizeSlowTaskWithGuid().AddLambda([this](FGuid TaskGuid, double TaskDuration)
+		GWarn->OnFinalizeSlowTaskWithGuid().AddLambda([this](FGuid TaskGuid, const FText& TaskName)
 			{
+				TRACE_END_REGION(*TaskName.ToString());
+
 				// Slow tasks can possibly be finalized from multiple threads, so we need to protect the registered span table
 				FScopeLock ScopeLock(&TaskSpanCriticalSection);
 
diff --git a/Engine/Plugins/Experimental/StudioTelemetry/Source/StudioTelemetry/Private/StudioTelemetryEditor.h b/Engine/Plugins/Experimental/StudioTelemetry/Source/StudioTelemetry/Private/StudioTelemetryEditor.h
--- a/Engine/Plugins/Experimental/StudioTelemetry/Source/StudioTelemetry/Private/StudioTelemetryEditor.h
+++ b/Engine/Plugins/Experimental/StudioTelemetry/Source/StudioTelemetry/Private/StudioTelemetryEditor.h
@@ -68,8 +68,8 @@ private:
 	const FName HitchingSpanName = TEXT("Hitching");
 	const FName OpenAssetEditorSpan = TEXT("Open Asset Editor");
 	const FName AssetRegistryScanSpanName = TEXT("Asset Registry Scan");
-	const float HeartbeatIntervalSeconds = 5.0;
-	const float HitchSamplerIntervalSeconds = 0.5;
+	const float HeartbeatIntervalSeconds = 1.0;
+	const float HitchSamplerIntervalSeconds = 0.1;
 	const float MinFPSForHitching = 5.0;
 
 	TMap<FGuid, TSharedPtr<IAnalyticsSpan>> TaskSpans;
diff --git a/Engine/Source/Runtime/Core/Public/Misc/FeedbackContext.h b/Engine/Source/Runtime/Core/Public/Misc/FeedbackContext.h
--- a/Engine/Source/Runtime/Core/Public/Misc/FeedbackContext.h
+++ b/Engine/Source/Runtime/Core/Public/Misc/FeedbackContext.h
@@ -58,7 +58,7 @@ public:
 	DECLARE_EVENT_TwoParams(FFeedbackContext, FOnStartSlowTaskWithGuid, FGuid Guid, const FText& TaskName);
 	FOnStartSlowTaskWithGuid& OnStartSlowTaskWithGuid() { return StartSlowTaskWithGuidEvent; }
 
-	DECLARE_EVENT_TwoParams(FFeedbackContext, FOnFinalizeSlowTaskWithGuid, FGuid Guid, double DurationInSeconds);
+	DECLARE_EVENT_TwoParams(FFeedbackContext, FOnFinalizeSlowTaskWithGuid, FGuid Guid, const FText& TaskName);
 	FOnFinalizeSlowTaskWithGuid& OnFinalizeSlowTaskWithGuid() { return FinalizeSlowTaskWithGuidEvent; }
 
 	/**** Legacy API - not deprecated as it's still in heavy use, but superceded by FScopedSlowTask ****/
@@ -91,7 +91,7 @@ protected:
 	{
 		const double TaskDuration = FPlatformTime::Seconds() - TaskStartTime;
 		FinalizeSlowTaskEvent.Broadcast(TaskName, TaskDuration);
-		FinalizeSlowTaskWithGuidEvent.Broadcast(TaskGuid, TaskDuration);
+		FinalizeSlowTaskWithGuidEvent.Broadcast(TaskGuid, TaskName);
 		GIsSlowTask = false;
 	}
 
