Added an option for users to premit/restrict sending of User data and Hardware data.
Adedd ability to turn off Telemetry capture altogether.

Telemetry is ON by default
User data is OFF by default
Hardware data is ON by default

Added commandline options for user to control what is sent by Telemetry:

ST_SendTelemetry=true
ST_SendUserData=false
ST_SendHardwareData=true

diff --git a/Engine/Plugins/Experimental/StudioTelemetry/Source/StudioTelemetry/Private/StudioTelemetry.cpp b/Engine/Plugins/Experimental/StudioTelemetry/Source/StudioTelemetry/Private/StudioTelemetry.cpp
--- a/Engine/Plugins/Experimental/StudioTelemetry/Source/StudioTelemetry/Private/StudioTelemetry.cpp
+++ b/Engine/Plugins/Experimental/StudioTelemetry/Source/StudioTelemetry/Private/StudioTelemetry.cpp
@@ -92,6 +92,15 @@ void FStudioTelemetry::EndSession()
 
 void FStudioTelemetry::StartSession()
 {
+	bool bSendTelemetry = true;  // Only send telemetry data if we have been requested to
+	FParse::Bool(FCommandLine::Get(), TEXT("ST_SendTelemetry="), bSendTelemetry);
+
+	if (bSendTelemetry == false)
+	{
+		// We did not wish to send any telemetry events
+		return;
+	}
+
 	AnalyticsProvider = FAnalyticsProviderMulticast::CreateAnalyticsProvider();
 
 	if (AnalyticsProvider.IsValid())
@@ -119,7 +128,13 @@ void FStudioTelemetry::StartSession()
 		FString SessionLabel;
 		FParse::Value(FCommandLine::Get(), TEXT("SessionLabel="), SessionLabel);
 
-		// Set the default event attributes	
+		bool bSendUserData = false;  // Never send user data unless specifically asked to
+		FParse::Bool(FCommandLine::Get(), TEXT("ST_SendUserData="), bSendUserData);
+
+		bool bSendHardwareData = true; // Always send hardware data unless specifically asked to
+		FParse::Bool(FCommandLine::Get(), TEXT("ST_SendHardwareData="), bSendHardwareData);
+
+		// Set the default event attributes, these will always be sent to telemetry for every event
 		DefaultEventAttributes.Emplace(TEXT("ProjectName"), ProjectName);
 		DefaultEventAttributes.Emplace(TEXT("ProjectID"), ProjectID);
 
@@ -134,19 +149,30 @@ void FStudioTelemetry::StartSession()
 		DefaultEventAttributes.Emplace(TEXT("Build_BranchName"), FApp::GetBranchName().ToLower());
 		DefaultEventAttributes.Emplace(TEXT("Build_Changelist"), BuildSettings::GetCurrentChangelist());
 
-		DefaultEventAttributes.Emplace(TEXT("Hardware_Platform"), FString(FPlatformProperties::IniPlatformName()));
-		DefaultEventAttributes.Emplace(TEXT("Hardware_GPU"), GRHIAdapterName);
-		DefaultEventAttributes.Emplace(TEXT("Hardware_CPU"), FPlatformMisc::GetCPUBrand());
-		DefaultEventAttributes.Emplace(TEXT("Hardware_CPU_Cores_Physical"), FPlatformMisc::NumberOfCores());
-		DefaultEventAttributes.Emplace(TEXT("Hardware_CPU_Cores_Logical"), FPlatformMisc::NumberOfCoresIncludingHyperthreads());
-		DefaultEventAttributes.Emplace(TEXT("Hardware_RAM"), static_cast<uint64>(FPlatformMemory::GetStats().TotalPhysical));
-		DefaultEventAttributes.Emplace(TEXT("Hardware_ComputerName"), ComputerName);
-
 		DefaultEventAttributes.Emplace(TEXT("Config_IsEditor"), GIsEditor);
 		DefaultEventAttributes.Emplace(TEXT("Config_IsUnattended"), FApp::IsUnattended());
 		DefaultEventAttributes.Emplace(TEXT("Config_IsBuildMachine"), GIsBuildMachine);
 		DefaultEventAttributes.Emplace(TEXT("Config_IsRunningCommandlet"), IsRunningCommandlet());
 		DefaultEventAttributes.Emplace(TEXT("Config_IsDebuggerPresent"), FPlatformMisc::IsDebuggerPresent());
+
+		// Only send user data if requested
+		if (bSendUserData == true)
+		{
+			DefaultEventAttributes.Emplace(TEXT("User_ID"), UserID);
+			DefaultEventAttributes.Emplace(TEXT("Application_Commandline"), FCommandLine::Get());
+		}
+
+		// Only send hardware data if requested
+		if (bSendHardwareData == true)
+		{
+			DefaultEventAttributes.Emplace(TEXT("Hardware_Platform"), FString(FPlatformProperties::IniPlatformName()));
+			DefaultEventAttributes.Emplace(TEXT("Hardware_GPU"), GRHIAdapterName);
+			DefaultEventAttributes.Emplace(TEXT("Hardware_CPU"), FPlatformMisc::GetCPUBrand());
+			DefaultEventAttributes.Emplace(TEXT("Hardware_CPU_Cores_Physical"), FPlatformMisc::NumberOfCores());
+			DefaultEventAttributes.Emplace(TEXT("Hardware_CPU_Cores_Logical"), FPlatformMisc::NumberOfCoresIncludingHyperthreads());
+			DefaultEventAttributes.Emplace(TEXT("Hardware_RAM"), static_cast<uint64>(FPlatformMemory::GetStats().TotalPhysical));
+			DefaultEventAttributes.Emplace(TEXT("Hardware_ComputerName"), ComputerName);
+		}
 		
 #if WITH_EDITOR
 		if (!FHorde::GetJobId().IsEmpty())
