Increment the PIEStartup span name for every subsequent transition to PIE after the first

diff --git a/Engine/Plugins/Experimental/StudioTelemetry/Source/StudioTelemetry/Private/StudioTelemetryEditor.cpp b/Engine/Plugins/Experimental/StudioTelemetry/Source/StudioTelemetry/Private/StudioTelemetryEditor.cpp
--- a/Engine/Plugins/Experimental/StudioTelemetry/Source/StudioTelemetry/Private/StudioTelemetryEditor.cpp
+++ b/Engine/Plugins/Experimental/StudioTelemetry/Source/StudioTelemetry/Private/StudioTelemetryEditor.cpp
@@ -655,8 +655,19 @@ void FStudioTelemetryEditor::Initialize()
 
 			// Start PIE span
 			PIESpan = FStudioTelemetry::Get().StartSpan(PIESpanName);
-			PIEStartupSpan = FStudioTelemetry::Get().StartSpan(PIEStartupSpanName, PIESpan);
-			
+
+			if (PIETransitionCount==0)
+			{ 
+				PIEStartupSpan = FStudioTelemetry::Get().StartSpan(PIEStartupSpanName, PIESpan);
+			}
+			else
+			{
+				// Append the PIE transition count to the PIE name
+				FNameBuilder NameBuilder(PIEStartupSpanName);
+				NameBuilder.Append(TEXT("%d"), PIETransitionCount);
+				PIEStartupSpan = FStudioTelemetry::Get().StartSpan(FName(NameBuilder), PIESpan);
+			}
+
 			TArray<FAnalyticsEventAttribute> Attributes;
 			Attributes.Emplace(TEXT("MapName"), EditorMapName);
 
