Fix virtualized asset analytics not reporting in cases where the process needs to be terminated afterwards.

#rnx

- Restores the functionality that was lost when CL 31207364 was backed out.
- Instead of a method to IAnalyticsProvider that can block until all analytics payloads have been sent we can make a normal flush on the analytics systems followed by a flush call to FHttpManager itself instead.
- Added code comments explaining why we need to call flush on FHttpManager.
- Including FHttpManager brought in 3 additional includes so reoordered the entire set of includes to match IWYU ordering (alphabetical) while I was adding them to the existing includes.

diff --git a/Engine/Plugins/Experimental/StudioTelemetry/Source/StudioTelemetry/Private/StudioTelemetryEditor.cpp b/Engine/Plugins/Experimental/StudioTelemetry/Source/StudioTelemetry/Private/StudioTelemetryEditor.cpp
--- a/Engine/Plugins/Experimental/StudioTelemetry/Source/StudioTelemetry/Private/StudioTelemetryEditor.cpp
+++ b/Engine/Plugins/Experimental/StudioTelemetry/Source/StudioTelemetry/Private/StudioTelemetryEditor.cpp
@@ -4,30 +4,33 @@
 
 #if WITH_EDITOR
 
-#include "StudioTelemetry.h"
 #include "AnalyticsTracer.h"
+#include "AssetRegistry/AssetRegistryTelemetry.h"
 #include "CollectionManagerModule.h"
 #include "ContentBrowserModule.h"
-#include "TelemetryRouter.h"
-#include "Policies/CondensedJsonPrintPolicy.h"
-#include "AssetRegistry/AssetRegistryTelemetry.h"
+#include "ContentBrowserTelemetry.h"
+#include "CookOnTheSide/CookOnTheFlyServer.h"
+#include "DerivedDataCacheInterface.h"
+#include "DerivedDataCacheUsageStats.h"
 #include "Editor.h"
 #include "Editor/UnrealEdEngine.h"
 #include "Engine/AssetManager.h"
-#include "UnrealEdGlobals.h"
-#include "CookOnTheSide/CookOnTheFlyServer.h"
-#include "Subsystems/AssetEditorSubsystem.h"
+#include "Experimental/ZenServerInterface.h"
+#include "FileHelpers.h"
+#include "HttpManager.h"
+#include "HttpModule.h"
+#include "IO/IoStoreOnDemand.h"
 #include "Misc/FeedbackContext.h"
+#include "Modules/ModuleManager.h"
+#include "Policies/CondensedJsonPrintPolicy.h"
 #include "ProfilingDebugging/CookStats.h"
-#include "DerivedDataCacheInterface.h"
-#include "DerivedDataCacheUsageStats.h"
-#include "Virtualization/VirtualizationSystem.h"
-#include "FileHelpers.h"
-#include "Experimental/ZenServerInterface.h"
-#include "ContentBrowserTelemetry.h"
 #include "ShaderStats.h"
+#include "StudioTelemetry.h"
+#include "Subsystems/AssetEditorSubsystem.h"
+#include "TelemetryRouter.h"
 #include "UObject/ICookInfo.h"
-#include "IO/IoStoreOnDemand.h"
+#include "UnrealEdGlobals.h"
+#include "Virtualization/VirtualizationSystem.h"
 
 UE_DISABLE_OPTIMIZATION_SHIP
 
@@ -973,11 +976,20 @@ void FStudioTelemetryEditor::Initialize()
 
 				if (EnumHasAllFlags(Flags, UE::Virtualization::EAnalyticsFlags::Flush))
 				{
-					// TODO: Flush currently does nothing while we rework the API
-					//if (TSharedPtr<IAnalyticsProvider> Provider = FStudioTelemetry::Get().GetProvider().Pin())
-					//{
-					//	Provider->BlockUntilFlushed(60.0f);
-					//}
+					if (TSharedPtr<IAnalyticsProvider> Provider = FStudioTelemetry::Get().GetProvider().Pin())
+					{
+						Provider->FlushEvents();
+
+						// It is quite likely that one of the analytics providers is sending data via the FHttpManager so we
+						// need to flush that as well to make sure that the message is sent. If the virtualization system has
+						// requested that we flush then it is quite likely that the process is to be terminated in which case
+						// not flushing the FHttpManager will result in analytics not being properly recorded.
+						if (FHttpModule* HttpModule = FModuleManager::GetModulePtr<FHttpModule>("HTTP"))
+						{
+							FHttpManager& HttpManager = HttpModule->GetHttpManager();
+							HttpManager.Flush(EHttpFlushReason::FullFlush);
+						}	
+					}
 				}
 			});
 	}
