Use  FApp::GetInstanceId() for the telemetry SessionID rather than creating a unique one. This will allow us to match Insights trace sessions to telemetry events.
Added Optional OS Information to default attributes

diff --git a/Engine/Plugins/Experimental/StudioTelemetry/Source/StudioTelemetry/Private/StudioTelemetry.cpp b/Engine/Plugins/Experimental/StudioTelemetry/Source/StudioTelemetry/Private/StudioTelemetry.cpp
--- a/Engine/Plugins/Experimental/StudioTelemetry/Source/StudioTelemetry/Private/StudioTelemetry.cpp
+++ b/Engine/Plugins/Experimental/StudioTelemetry/Source/StudioTelemetry/Private/StudioTelemetry.cpp
@@ -122,8 +122,7 @@ void FStudioTelemetry::StartSession()
 			ProjectID = FGuid(FCString::Atoi(*(Elements[1])), FCString::Atoi(*(Elements[2])), FCString::Atoi(*(Elements[3])), FCString::Atoi(*(Elements[4])));
 		}
 
-		FGuid SessionID;
-		FPlatformMisc::CreateGuid(SessionID);
+		FGuid SessionID = FApp::GetInstanceId();
 		
 		FString SessionLabel;
 		FParse::Value(FCommandLine::Get(), TEXT("SessionLabel="), SessionLabel);
@@ -131,9 +130,12 @@ void FStudioTelemetry::StartSession()
 		bool bSendUserData = false;  // Never send user data unless specifically asked to
 		FParse::Bool(FCommandLine::Get(), TEXT("ST_SendUserData="), bSendUserData);
 
-		bool bSendHardwareData = true; // Always send hardware data unless specifically asked to
+		bool bSendHardwareData = true; // Always send hardware data unless specifically asked not to
 		FParse::Bool(FCommandLine::Get(), TEXT("ST_SendHardwareData="), bSendHardwareData);
 
+		bool bSendOSData = true; // Always send operating system data unless specifically asked not to
+		FParse::Bool(FCommandLine::Get(), TEXT("ST_SendOSData="), bSendOSData);
+
 		// Set the default event attributes, these will always be sent to telemetry for every event
 		DefaultEventAttributes.Emplace(TEXT("ProjectName"), ProjectName);
 		DefaultEventAttributes.Emplace(TEXT("ProjectID"), ProjectID);
@@ -173,6 +175,20 @@ void FStudioTelemetry::StartSession()
 			DefaultEventAttributes.Emplace(TEXT("Hardware_RAM"), static_cast<uint64>(FPlatformMemory::GetStats().TotalPhysical));
 			DefaultEventAttributes.Emplace(TEXT("Hardware_ComputerName"), ComputerName);
 		}
+
+		// Only send OS data if requested
+		if (bSendOSData==true)
+		{
+			FString OSVersionLabel;
+			FString OSSubVersionLabel;
+
+			FPlatformMisc::GetOSVersions(OSVersionLabel, OSSubVersionLabel);
+
+			DefaultEventAttributes.Emplace(TEXT("OS_Version"), FPlatformMisc::GetOSVersion());
+			DefaultEventAttributes.Emplace(TEXT("OS_VersionLabel"), OSVersionLabel);
+			DefaultEventAttributes.Emplace(TEXT("OS_VersionSubLabel"), OSSubVersionLabel);
+			DefaultEventAttributes.Emplace(TEXT("OS_ID"), FPlatformMisc::GetOperatingSystemId());	
+		}
 		
 #if WITH_EDITOR
 		if (!FHorde::GetJobId().IsEmpty())
@@ -278,3 +294,4 @@ TSharedPtr<IAnalyticsSpan> FStudioTelemetry::GetSpan(const FName Name)
 {
 	return AnalyticsTracer.IsValid() ? AnalyticsTracer->GetSpan(Name) : TSharedPtr<IAnalyticsSpan>();
 }
+
