Studio Telemetry Configuration Changes
Moved the provider cofiguration settings to the StudioTelemetry.Provider.ProviderName section
Added StudioTelemetry global configuration settings to the StudioTelemetry.Config section

diff --git a/Engine/Config/BaseEngine.ini b/Engine/Config/BaseEngine.ini
--- a/Engine/Config/BaseEngine.ini
+++ b/Engine/Config/BaseEngine.ini
@@ -3552,9 +3552,14 @@ IntervalSeconds=1
 +Listeners=(MetricName=AddedConnections, Class=/Script/Engine.NetworkMetricsPerfCounters)
 
 ; Studio Telemetry Settings
+[StudioTelemetry.Config] 
+SendTelemetry=true
+SendUserData=false
+SendHardwareData=true
+SendOSData=true
 
 ; Example Log File Analytics
-[StudioTelemetry.LogAnalytics] 
+[StudioTelemetry.Provider.LogAnalytics] 
 Name=LogAnalytics
 ProviderModule=AnalyticsLog
 UsageType=Editor
@@ -3563,7 +3568,7 @@ UsageType=Editor
 ; FolderPath="somefolder"
 
 ; Example Horde Analytics Settings
-;[StudioTelemetry.HordeAnalytics]
+;[StudioTelemetry.Provider.HordeAnalytics]
 ; Name=HordeAnalytics
 ; ProviderModule=AnalyticsET
 ; UsageType=EditorAndClient
diff --git a/Engine/Plugins/Experimental/StudioTelemetry/Source/StudioTelemetry/Private/AnalyticsProviderMulticast.cpp b/Engine/Plugins/Experimental/StudioTelemetry/Source/StudioTelemetry/Private/AnalyticsProviderMulticast.cpp
--- a/Engine/Plugins/Experimental/StudioTelemetry/Source/StudioTelemetry/Private/AnalyticsProviderMulticast.cpp
+++ b/Engine/Plugins/Experimental/StudioTelemetry/Source/StudioTelemetry/Private/AnalyticsProviderMulticast.cpp
@@ -8,7 +8,6 @@
 #include "Misc/ConfigCacheIni.h"
 #include "HttpModule.h"
 
-const FString TelemetrySection(TEXT("StudioTelemetry"));
 static FString ProviderSection;
 
 FString GetAnalyticsProviderConfiguration(const FString& Name, bool)
@@ -31,13 +30,15 @@ TWeakPtr<IAnalyticsProvider> FAnalyticsProviderMulticast::GetAnalyticsProvider(c
 
 FAnalyticsProviderMulticast::FAnalyticsProviderMulticast()
 {
+	const FString TelemetryProviderSection(TEXT("StudioTelemetry.Provider"));
+
 	TArray<FString> SectionNames;
 	
 	if (GConfig->GetSectionNames(GEngineIni, SectionNames))
 	{
 		for (const FString& SectionName : SectionNames)
 		{
-			if (SectionName.Find(TelemetrySection) != INDEX_NONE)
+			if (SectionName.Find(TelemetryProviderSection) != INDEX_NONE)
 			{
 				ProviderSection = SectionName;
 
diff --git a/Engine/Plugins/Experimental/StudioTelemetry/Source/StudioTelemetry/Private/StudioTelemetry.cpp b/Engine/Plugins/Experimental/StudioTelemetry/Source/StudioTelemetry/Private/StudioTelemetry.cpp
--- a/Engine/Plugins/Experimental/StudioTelemetry/Source/StudioTelemetry/Private/StudioTelemetry.cpp
+++ b/Engine/Plugins/Experimental/StudioTelemetry/Source/StudioTelemetry/Private/StudioTelemetry.cpp
@@ -50,6 +50,9 @@ void FStudioTelemetry::StartupModule()
 {
 	UE_LOG(LogStudioTelemetry, Display, TEXT("Starting StudioTelemetry Module"));
 
+	// Load the configureation
+	FStudioTelemetry::Get().LoadConfiguration();
+
 	// Create the provider and start the analytics session
 	FStudioTelemetry::Get().StartSession();
 
@@ -90,12 +93,37 @@ void FStudioTelemetry::EndSession()
 	}
 }
 
-void FStudioTelemetry::StartSession()
+void FStudioTelemetry::LoadConfiguration()
 {
-	bool bSendTelemetry = true;  // Only send telemetry data if we have been requested to
-	FParse::Bool(FCommandLine::Get(), TEXT("ST_SendTelemetry="), bSendTelemetry);
+	const FString TelemetryConfigurationSection(TEXT("StudioTelemetry.Config"));
 
-	if (bSendTelemetry == false)
+	// Look for the configuration seetings in the Engine.ini files
+	TArray<FString> SectionNames;
+
+	if (GConfig->GetSectionNames(GEngineIni, SectionNames))
+	{
+		for (const FString& SectionName : SectionNames)
+		{
+			if (SectionName.Find(TelemetryConfigurationSection) != INDEX_NONE)
+			{
+				GConfig->GetBool(*SectionName, TEXT("SendTelemetry="), Config.bSendTelemetry, GEngineIni);
+				GConfig->GetBool(*SectionName, TEXT("SendUserData="), Config.bSendUserData, GEngineIni);
+				GConfig->GetBool(*SectionName, TEXT("SendHardwareData="), Config.bSendHardwareData, GEngineIni);
+				GConfig->GetBool(*SectionName, TEXT("SendOSData="), Config.bSendOSData, GEngineIni);
+			}
+		}
+	}
+
+	// Parse the comamndline for any local configuration overrides
+	FParse::Bool(FCommandLine::Get(), TEXT("ST_SendTelemetry="), Config.bSendTelemetry);
+	FParse::Bool(FCommandLine::Get(), TEXT("ST_SendUserData="), Config.bSendUserData);
+	FParse::Bool(FCommandLine::Get(), TEXT("ST_SendHardwareData="), Config.bSendHardwareData);
+	FParse::Bool(FCommandLine::Get(), TEXT("ST_SendOSData="), Config.bSendOSData);
+}
+
+void FStudioTelemetry::StartSession()
+{
+	if (Config.bSendTelemetry == false)
 	{
 		// We did not wish to send any telemetry events
 		return;
@@ -127,15 +155,6 @@ void FStudioTelemetry::StartSession()
 		FString SessionLabel;
 		FParse::Value(FCommandLine::Get(), TEXT("SessionLabel="), SessionLabel);
 
-		bool bSendUserData = false;  // Never send user data unless specifically asked to
-		FParse::Bool(FCommandLine::Get(), TEXT("ST_SendUserData="), bSendUserData);
-
-		bool bSendHardwareData = true; // Always send hardware data unless specifically asked not to
-		FParse::Bool(FCommandLine::Get(), TEXT("ST_SendHardwareData="), bSendHardwareData);
-
-		bool bSendOSData = true; // Always send operating system data unless specifically asked not to
-		FParse::Bool(FCommandLine::Get(), TEXT("ST_SendOSData="), bSendOSData);
-
 		// Set the default event attributes, these will always be sent to telemetry for every event
 		DefaultEventAttributes.Emplace(TEXT("ProjectName"), ProjectName);
 		DefaultEventAttributes.Emplace(TEXT("ProjectID"), ProjectID);
@@ -158,14 +177,14 @@ void FStudioTelemetry::StartSession()
 		DefaultEventAttributes.Emplace(TEXT("Config_IsDebuggerPresent"), FPlatformMisc::IsDebuggerPresent());
 
 		// Only send user data if requested
-		if (bSendUserData == true)
+		if (Config.bSendUserData == true)
 		{
 			DefaultEventAttributes.Emplace(TEXT("User_ID"), UserID);
 			DefaultEventAttributes.Emplace(TEXT("Application_Commandline"), FCommandLine::Get());
 		}
 
 		// Only send hardware data if requested
-		if (bSendHardwareData == true)
+		if (Config.bSendHardwareData == true)
 		{
 			DefaultEventAttributes.Emplace(TEXT("Hardware_Platform"), FString(FPlatformProperties::IniPlatformName()));
 			DefaultEventAttributes.Emplace(TEXT("Hardware_GPU"), GRHIAdapterName);
@@ -177,7 +196,7 @@ void FStudioTelemetry::StartSession()
 		}
 
 		// Only send OS data if requested
-		if (bSendOSData==true)
+		if (Config.bSendOSData==true)
 		{
 			FString OSVersionLabel;
 			FString OSSubVersionLabel;
diff --git a/Engine/Plugins/Experimental/StudioTelemetry/Source/StudioTelemetry/Public/StudioTelemetry.h b/Engine/Plugins/Experimental/StudioTelemetry/Source/StudioTelemetry/Public/StudioTelemetry.h
--- a/Engine/Plugins/Experimental/StudioTelemetry/Source/StudioTelemetry/Public/StudioTelemetry.h
+++ b/Engine/Plugins/Experimental/StudioTelemetry/Source/StudioTelemetry/Public/StudioTelemetry.h
@@ -105,11 +105,25 @@ private:
 	/** Ends an existing analytics session*/
 	void EndSession();
 
+	/** Configure the plugin*/
+	void LoadConfiguration();
+
+	struct FConfig
+	{
+		bool bSendTelemetry = true; // Only send telemetry data if we have been requested to
+		bool bSendUserData = false;  // Never send user data unless specifically asked to
+		bool bSendHardwareData = true; // Always send hardware data unless specifically asked not to
+		bool bSendOSData = true; // Always send operating system data unless specifically asked not to
+	};
+	
 	FCriticalSection						CriticalSection;
 	TSharedPtr<FAnalyticsProviderMulticast>	AnalyticsProvider;
 	TSharedPtr<IAnalyticsTracer>			AnalyticsTracer;
 	OnRecordEvent							RecordEventCallback;
 	FGuid									SessionGUID;
+	FConfig									Config;
+
+	
 };
 
 #define STUDIO_TELEMETRY_SPAN_SCOPE(Name) FStudioTelemetry::ScopedSpan PREPROCESSOR_JOIN(ScopedSpan, __LINE__)(TEXT(#Name));
