Fixed typo in Studio Telemetry

diff --git a/Engine/Plugins/Experimental/StudioTelemetry/Source/StudioTelemetry/Private/StudioTelemetry.cpp b/Engine/Plugins/Experimental/StudioTelemetry/Source/StudioTelemetry/Private/StudioTelemetry.cpp
--- a/Engine/Plugins/Experimental/StudioTelemetry/Source/StudioTelemetry/Private/StudioTelemetry.cpp
+++ b/Engine/Plugins/Experimental/StudioTelemetry/Source/StudioTelemetry/Private/StudioTelemetry.cpp
@@ -114,7 +114,7 @@ void FStudioTelemetry::LoadConfiguration()
 		}
 	}
 
-	// Parse the comamndline for any local configuration overrides
+	// Parse the commandline for any local configuration overrides
 	FParse::Bool(FCommandLine::Get(), TEXT("ST_SendTelemetry="), Config.bSendTelemetry);
 	FParse::Bool(FCommandLine::Get(), TEXT("ST_SendUserData="), Config.bSendUserData);
 	FParse::Bool(FCommandLine::Get(), TEXT("ST_SendHardwareData="), Config.bSendHardwareData);
