Added Horde Analytics Backend Module which can automatically set the Horde Server URL if one is present in the Horde Environment Variable
Horde Analytics Backend essentially creates the same IAnalyticsProviderET interface but
Added GetServerURL to FHorde interface which retrieves the Horde Environment Variable
Changed the Horde Analytics backends in teh config files to use the new HordeAnalytics backend  instead of the AnalyticsET backend.

diff --git a/Engine/Plugins/Experimental/StudioTelemetry/Source/AnalyticsHorde/AnalyticsHorde.Build.cs b/Engine/Plugins/Experimental/StudioTelemetry/Source/AnalyticsHorde/AnalyticsHorde.Build.cs
new file mode 100644
--- /dev/null
+++ b/Engine/Plugins/Experimental/StudioTelemetry/Source/AnalyticsHorde/AnalyticsHorde.Build.cs
@@ -0,0 +1,25 @@
+// Copyright Epic Games, Inc. All Rights Reserved.
+
+namespace UnrealBuildTool.Rules
+{
+	public class AnalyticsHorde : ModuleRules
+	{
+		public AnalyticsHorde(ReadOnlyTargetRules Target) : base(Target)
+		{
+			PublicDependencyModuleNames.AddRange(
+				new string[] {
+						"Core"
+				}
+			);
+
+			PrivateDependencyModuleNames.AddRange(
+				new string[] {
+				"CoreUObject",
+				"Analytics",
+				"AnalyticsET",
+				"Horde"
+				}
+			);
+		}
+	}
+}
diff --git a/Engine/Plugins/Experimental/StudioTelemetry/Source/AnalyticsHorde/Private/AnalyticsHorde.cpp b/Engine/Plugins/Experimental/StudioTelemetry/Source/AnalyticsHorde/Private/AnalyticsHorde.cpp
new file mode 100644
--- /dev/null
+++ b/Engine/Plugins/Experimental/StudioTelemetry/Source/AnalyticsHorde/Private/AnalyticsHorde.cpp
@@ -0,0 +1,37 @@
+// Copyright Epic Games, Inc. All Rights Reserved.
+
+#include "AnalyticsHorde.h"
+#include "AnalyticsET.h"
+#include "IAnalyticsProviderET.h"
+#include "Horde.h"
+
+IMPLEMENT_MODULE( FAnalyticsHorde, AnalyticsHorde );
+
+void FAnalyticsHorde::StartupModule()
+{
+}
+
+void FAnalyticsHorde::ShutdownModule()
+{
+}
+
+TSharedPtr<IAnalyticsProvider> FAnalyticsHorde::CreateAnalyticsProvider(const FAnalyticsProviderConfigurationDelegate& GetConfigValue) const
+{
+	if (GetConfigValue.IsBound())
+	{
+		TSharedPtr<IAnalyticsProviderET> AnalyicsProviderET = FAnalyticsET::Get().CreateAnalyticsProviderET(GetConfigValue);
+
+		// Check if we have specifcied a Horde URL in the environment. This allows jobs running on Horde to send telemetry directly to the server they were run on
+		const FString HordeServerURL = FHorde::GetServerURL();
+
+		if (AnalyicsProviderET.IsValid() && !HordeServerURL.IsEmpty())
+		{
+			TArray<FString> AltDomains;
+			AnalyicsProviderET->SetURLEndpoint(HordeServerURL, AltDomains);
+		}
+		
+		return AnalyicsProviderET;
+	}
+
+	return nullptr;
+}
diff --git a/Engine/Plugins/Experimental/StudioTelemetry/Source/AnalyticsHorde/Public/AnalyticsHorde.h b/Engine/Plugins/Experimental/StudioTelemetry/Source/AnalyticsHorde/Public/AnalyticsHorde.h
new file mode 100644
--- /dev/null
+++ b/Engine/Plugins/Experimental/StudioTelemetry/Source/AnalyticsHorde/Public/AnalyticsHorde.h
@@ -0,0 +1,50 @@
+// Copyright Epic Games, Inc. All Rights Reserved.
+
+#pragma once
+
+#include "AnalyticsProviderConfigurationDelegate.h"
+#include "Interfaces/IAnalyticsProviderModule.h"
+#include "Modules/ModuleManager.h"
+
+class IAnalyticsProvider;
+
+/**
+ *  Public implementation of IAnalyticsProviderModule
+ */
+class FAnalyticsHorde : public IAnalyticsProviderModule
+{
+	//--------------------------------------------------------------------------
+	// Module functionality
+	//--------------------------------------------------------------------------
+public:
+	/**
+	 * Singleton-like access to this module's interface.  This is just for convenience!
+	 * Beware of calling this during the shutdown phase, though.  Your module might have been unloaded already.
+	 *
+	 * @return Returns singleton instance, loading the module on demand if needed
+	 */
+	static inline FAnalyticsHorde& Get()
+	{
+		return FModuleManager::LoadModuleChecked< FAnalyticsHorde >( "AnalyticsHorde" );
+	}
+
+	//--------------------------------------------------------------------------
+	// Configuration functionality
+	//--------------------------------------------------------------------------
+public:
+	
+	//--------------------------------------------------------------------------
+	// provider factory functions
+	//--------------------------------------------------------------------------
+public:
+	/**
+	 * IAnalyticsProviderModule interface.
+	 * Creates the analytics provider given a configuration delegate.
+	 * The keys required exactly match the field names in the Config object. 
+	 */
+	ANALYTICSHORDE_API virtual TSharedPtr<IAnalyticsProvider> CreateAnalyticsProvider(const FAnalyticsProviderConfigurationDelegate& GetConfigValue) const override;
+
+private:
+	ANALYTICSHORDE_API virtual void StartupModule() override;
+	ANALYTICSHORDE_API virtual void ShutdownModule() override;
+};
diff --git a/Engine/Plugins/Experimental/StudioTelemetry/Source/StudioTelemetry/StudioTelemetry.build.cs b/Engine/Plugins/Experimental/StudioTelemetry/Source/StudioTelemetry/StudioTelemetry.build.cs
--- a/Engine/Plugins/Experimental/StudioTelemetry/Source/StudioTelemetry/StudioTelemetry.build.cs
+++ b/Engine/Plugins/Experimental/StudioTelemetry/Source/StudioTelemetry/StudioTelemetry.build.cs
@@ -44,13 +44,6 @@ public class StudioTelemetry : ModuleRules
 					"TelemetryUtils"
 				}
 			);
-
-			PublicDependencyModuleNames.AddRange(
-				new string[] {
-					"AnalyticsET",
-					"AnalyticsLog"
-				}
-			);
 		}
 	}
 }
diff --git a/Engine/Plugins/Experimental/StudioTelemetry/StudioTelemetry.uplugin b/Engine/Plugins/Experimental/StudioTelemetry/StudioTelemetry.uplugin
--- a/Engine/Plugins/Experimental/StudioTelemetry/StudioTelemetry.uplugin
+++ b/Engine/Plugins/Experimental/StudioTelemetry/StudioTelemetry.uplugin
@@ -20,6 +20,11 @@
 			"Type": "EditorAndProgram",
 			"LoadingPhase": "PostConfigInit"
 		},
+		{
+			"Name": "AnalyticsHorde",
+			"Type": "EditorAndProgram",
+			"LoadingPhase": "PostConfigInit"
+		},
 		{
 			"Name": "StudioTelemetry",
 			"Type": "EditorAndProgram",
diff --git a/Engine/Source/Developer/Horde/Private/Horde.cpp b/Engine/Source/Developer/Horde/Private/Horde.cpp
--- a/Engine/Source/Developer/Horde/Private/Horde.cpp
+++ b/Engine/Source/Developer/Horde/Private/Horde.cpp
@@ -129,4 +129,18 @@ FString FHorde::GetBatchId()
 		}
 	}
 	return BatchId;
-} 
\ No newline at end of file
+} 
+
+FString FHorde::GetServerURL()
+{
+	static FString ServerURL;
+
+	if (ServerURL.IsEmpty())
+	{
+		if (false == FParse::Value(FCommandLine::Get(), TEXT("HordeServerUrl="), ServerURL))
+		{
+			ServerURL = FPlatformMisc::GetEnvironmentVariable(TEXT("UE_HORDE_URL"));
+		}
+	}
+	return ServerURL;
+}
\ No newline at end of file
diff --git a/Engine/Source/Developer/Horde/Public/Horde.h b/Engine/Source/Developer/Horde/Public/Horde.h
--- a/Engine/Source/Developer/Horde/Public/Horde.h
+++ b/Engine/Source/Developer/Horde/Public/Horde.h
@@ -12,6 +12,7 @@ struct FHorde
 {
 	HORDE_API static FString GetTemplateName();
 	HORDE_API static FString GetTemplateId();
+	HORDE_API static FString GetServerURL();
 	HORDE_API static FString GetJobId();
 	HORDE_API static FString GetJobURL();
 	HORDE_API static FString GetStepId();
