Added Horde Server URL to Studio Telemetry default attributes

diff --git a/Engine/Plugins/Experimental/StudioTelemetry/Source/StudioTelemetry/Private/StudioTelemetry.cpp b/Engine/Plugins/Experimental/StudioTelemetry/Source/StudioTelemetry/Private/StudioTelemetry.cpp
--- a/Engine/Plugins/Experimental/StudioTelemetry/Source/StudioTelemetry/Private/StudioTelemetry.cpp
+++ b/Engine/Plugins/Experimental/StudioTelemetry/Source/StudioTelemetry/Private/StudioTelemetry.cpp
@@ -213,6 +213,7 @@ void FStudioTelemetry::StartSession()
 		if (!FHorde::GetJobId().IsEmpty())
 		{
 			// Only send Horde data if applicable
+			DefaultEventAttributes.Emplace(TEXT("Horde_ServerURL"), FHorde::GetServerURL());
 			DefaultEventAttributes.Emplace(TEXT("Horde_TemplateID"), FHorde::GetTemplateId());
 			DefaultEventAttributes.Emplace(TEXT("Horde_TemplateName"), FHorde::GetTemplateName());
 			DefaultEventAttributes.Emplace(TEXT("Horde_JobURL"), FHorde::GetJobURL());
