Fix for shipping/non-editor builds. Horde dependency is not applicable to shipping/non-editor builds.

diff --git a/Engine/Plugins/Experimental/StudioTelemetry/Source/StudioTelemetry/StudioTelemetry.build.cs b/Engine/Plugins/Experimental/StudioTelemetry/Source/StudioTelemetry/StudioTelemetry.build.cs
--- a/Engine/Plugins/Experimental/StudioTelemetry/Source/StudioTelemetry/StudioTelemetry.build.cs
+++ b/Engine/Plugins/Experimental/StudioTelemetry/Source/StudioTelemetry/StudioTelemetry.build.cs
@@ -12,8 +12,7 @@ public class StudioTelemetry : ModuleRules
 		PublicDependencyModuleNames.AddRange(
 			new string[] {
 				"Core",
-				"HTTP",
-				"Horde"
+				"HTTP"
 			}
 		);
 
@@ -28,5 +27,14 @@ public class StudioTelemetry : ModuleRules
 				"RHI",
 			}
 		);
+
+		if (Target.bBuildEditor)
+		{
+			PrivateDependencyModuleNames.AddRange(
+				new string[] {
+					"Horde"
+				}
+			);
+		}
 	}
 }
