Disable Studio Telemetry in Shipping builds

diff --git a/Engine/Plugins/Experimental/StudioTelemetry/Source/StudioTelemetry/Private/StudioTelemetry.cpp b/Engine/Plugins/Experimental/StudioTelemetry/Source/StudioTelemetry/Private/StudioTelemetry.cpp
--- a/Engine/Plugins/Experimental/StudioTelemetry/Source/StudioTelemetry/Private/StudioTelemetry.cpp
+++ b/Engine/Plugins/Experimental/StudioTelemetry/Source/StudioTelemetry/Private/StudioTelemetry.cpp
@@ -47,6 +47,7 @@ void FStudioTelemetry::SetRecordEventCallback(OnRecordEventCallback Callback )
 
 void FStudioTelemetry::StartupModule()
 {
+#if !UE_BUILD_SHIPPING
 	UE_LOG(LogStudioTelemetry, Display, TEXT("Starting StudioTelemetry Module"));
 
 	// Load the configuration
@@ -54,14 +55,17 @@ void FStudioTelemetry::StartupModule()
 
 	// Create the provider and start the analytics session
 	FStudioTelemetry::Get().StartSession();
+#endif
 }
 
 void FStudioTelemetry::ShutdownModule()
 {
+#if !UE_BUILD_SHIPPING
 	// End the session and destroy analytics provider
 	FStudioTelemetry::Get().EndSession();
 
 	UE_LOG(LogStudioTelemetry, Display, TEXT("Shutdown StudioTelemetry Module"));
+#endif
 }
 
 void FStudioTelemetry::EndSession()
