Support plain-text ProjectID GUIDs in StudioTelemetry from DefaultGame.ini.

diff --git a/Engine/Plugins/Experimental/StudioTelemetry/Source/StudioTelemetry/Private/StudioTelemetry.cpp b/Engine/Plugins/Experimental/StudioTelemetry/Source/StudioTelemetry/Private/StudioTelemetry.cpp
--- a/Engine/Plugins/Experimental/StudioTelemetry/Source/StudioTelemetry/Private/StudioTelemetry.cpp
+++ b/Engine/Plugins/Experimental/StudioTelemetry/Source/StudioTelemetry/Private/StudioTelemetry.cpp
@@ -139,10 +139,17 @@ void FStudioTelemetry::StartSession()
 
 		FGuid ProjectID;
 
-		TArray<FString> Elements;
-		if (ProjectIDString.ParseIntoArray(Elements, TEXT("=")) == 5) 
+		if (!ProjectIDString.IsEmpty())
 		{
-			ProjectID = FGuid(FCString::Atoi(*(Elements[1])), FCString::Atoi(*(Elements[2])), FCString::Atoi(*(Elements[3])), FCString::Atoi(*(Elements[4])));
+			TArray<FString> Elements;
+			if (ProjectIDString.ParseIntoArray(Elements, TEXT("=")) == 5)
+			{
+				ProjectID = FGuid(FCString::Atoi(*(Elements[1])), FCString::Atoi(*(Elements[2])), FCString::Atoi(*(Elements[3])), FCString::Atoi(*(Elements[4])));
+			}
+			else
+			{
+				ProjectID = FGuid(ProjectIDString);
+			}
 		}
 
 		FGuid SessionID = FApp::GetInstanceId();
