Changed log messages to Warnings if the provider configuration is malformed. StudioTelemetry should never cause an error in critical builds.

diff --git a/Engine/Plugins/Experimental/StudioTelemetry/Source/StudioTelemetry/Private/AnalyticsProviderMulticast.cpp b/Engine/Plugins/Experimental/StudioTelemetry/Source/StudioTelemetry/Private/AnalyticsProviderMulticast.cpp
--- a/Engine/Plugins/Experimental/StudioTelemetry/Source/StudioTelemetry/Private/AnalyticsProviderMulticast.cpp
+++ b/Engine/Plugins/Experimental/StudioTelemetry/Source/StudioTelemetry/Private/AnalyticsProviderMulticast.cpp
@@ -70,7 +70,7 @@ FAnalyticsProviderMulticast::FAnalyticsProviderMulticast()
 				else
 				{
 					// Must always specify a usage type
-					UE_LOG(LogStudioTelemetry, Error, TEXT("There must be a valid UsageType specified for analytics provider %s"), *ProviderSection);
+					UE_LOG(LogStudioTelemetry, Warning, TEXT("There must be a valid UsageType specified for analytics provider %s"), *ProviderSection);
 					continue;
 				}
 
@@ -108,7 +108,7 @@ FAnalyticsProviderMulticast::FAnalyticsProviderMulticast()
 				}
 				else
 				{
-					UE_LOG(LogStudioTelemetry, Error, TEXT("There must be a valid ProviderModule specified for analytics provider %s"), *ProviderSection);
+					UE_LOG(LogStudioTelemetry, Warning, TEXT("There must be a valid ProviderModule specified for analytics provider %s"), *ProviderSection);
 				}
 			}
 		}
