Fix an issue where Studio Telemetry could cause extreme performance issues in the editor.

- FAnalyticsProviderLog::RecordEvent was writing analytics payloads to a file handle and flushing that output to disk. These operations could end up being very slow especially the flushing call, depending on the hardware and other disk operations being run on the system by other processes.
- This by itself could cause minor hitches as RecordEvent will always be called on the GameThread, but things would go really bad if the hitch was long enough to cause the hitch tracker to fire, which in turn would record a new analytics event for the hitch and if the flush was slow enough this would invoke another hitch and so on. So if the slow down was consistent and high enough the editor could potentially get into a doom loop of hitches, where recording the hitch would cause another hitch and so on.
- FAnalyticsProviderLog now creates an FPipe when enabled and all writes to FileWriter are pushed as jogs to that pipe, so even if the flush ends up being slow it will be off of the GameThread and not cause a hitch itself.
-- Note that the use of FPipe will keep the jobs sequential so that the analytics file does not get out of order. FArchive is not thread safe so we need to be sure that the jobs only run one at a time anyway.
- Added trace scopes to make issues in this code easier to track in future.

diff --git a/Engine/Plugins/Experimental/StudioTelemetry/Source/AnalyticsLog/Private/AnalyticsProviderLog.cpp b/Engine/Plugins/Experimental/StudioTelemetry/Source/AnalyticsLog/Private/AnalyticsProviderLog.cpp
--- a/Engine/Plugins/Experimental/StudioTelemetry/Source/AnalyticsLog/Private/AnalyticsProviderLog.cpp
+++ b/Engine/Plugins/Experimental/StudioTelemetry/Source/AnalyticsLog/Private/AnalyticsProviderLog.cpp
@@ -1,9 +1,11 @@
 // Copyright Epic Games, Inc. All Rights Reserved.
 
 #include "AnalyticsProviderLog.h"
+
 #include "Analytics.h"
 #include "HAL/FileManager.h"
 #include "Misc/Paths.h"
+#include "Tasks/Pipe.h"
 
 FAnalyticsProviderLog::FAnalyticsProviderLog(const FAnalyticsProviderConfigurationDelegate& GetConfigValue)
 {
@@ -26,10 +28,16 @@ FAnalyticsProviderLog::FAnalyticsProviderLog(const FAnalyticsProviderConfigurati
 	// Create the full output path
 	FString FilePath = FolderPath / FileName;
 	FileWriter = TUniquePtr<FArchive>(IFileManager::Get().CreateFileWriter(*FilePath, FILEWRITE_EvenIfReadOnly));
+
+	if (FileWriter.IsValid())
+	{
+		WriterPipe = MakeUnique<UE::Tasks::FPipe>(TEXT("FAnalyticsProviderLog_Writer"));
+	}
 }
 
 FAnalyticsProviderLog::~FAnalyticsProviderLog()
 {
+
 }
 
 bool FAnalyticsProviderLog::SetSessionID(const FString& InSessionID)
@@ -80,6 +88,8 @@ FAnalyticsEventAttribute FAnalyticsProviderLog::GetDefaultEventAttribute(int Att
 
 bool FAnalyticsProviderLog::StartSession(const TArray<FAnalyticsEventAttribute>& Attributes)
 {
+	TRACE_CPUPROFILER_EVENT_SCOPE(FAnalyticsProviderLog::StartSession);
+
 	RecordEvent(TEXT("StartSession"), Attributes);
 
 	return true;
@@ -87,20 +97,33 @@ bool FAnalyticsProviderLog::StartSession(const TArray<FAnalyticsEventAttribute>&
 
 void FAnalyticsProviderLog::EndSession()
 {
+	TRACE_CPUPROFILER_EVENT_SCOPE(FAnalyticsProviderLog::EndSession);
+
 	RecordEvent(TEXT("EndSession"));
 
-	if (FileWriter)
+	if (WriterPipe.IsValid())
+	{
+		TRACE_CPUPROFILER_EVENT_SCOPE(FAnalyticsProviderLog::EndSession::ClosePipe);
+		WriterPipe->WaitUntilEmpty();
+		WriterPipe.Reset();
+	}
+
+	if (FileWriter.IsValid())
 	{
 		FileWriter->Flush();
 		FileWriter->Close();
+
+		FileWriter.Reset();
 	}
 }
 
 void FAnalyticsProviderLog::RecordEvent(const FString& EventName, const TArray<FAnalyticsEventAttribute>& Attributes)
 {	
+	TRACE_CPUPROFILER_EVENT_SCOPE(FAnalyticsProviderLog::RecordEvent);
+
 	static uint32 RecordId(0);
 
-	if (FileWriter)
+	if (FileWriter.IsValid() && WriterPipe.IsValid())
 	{
 		TStringBuilder<1024> Builder;
 
@@ -131,7 +154,16 @@ void FAnalyticsProviderLog::RecordEvent(const FString& EventName, const TArray<F
 			}
 		}
 
-		FileWriter->Logf(TEXT("%s}"),Builder.ToString());
-		FileWriter->Flush();
+		Builder.Append(TEXT("}"));
+
+		FString FinalOutput = Builder.ToString();
+
+		WriterPipe->Launch(TEXT("FAnalyticsProviderLog_WriteJob"), [this, FinalOutput = MoveTemp(FinalOutput)]()
+			{
+				TRACE_CPUPROFILER_EVENT_SCOPE(FAnalyticsProviderLog_WriteJob);
+				
+				FileWriter->Logf(TEXT("%s"), *FinalOutput);
+				FileWriter->Flush();
+			});
 	}
-}
\ No newline at end of file
+}
diff --git a/Engine/Plugins/Experimental/StudioTelemetry/Source/AnalyticsLog/Private/AnalyticsProviderLog.h b/Engine/Plugins/Experimental/StudioTelemetry/Source/AnalyticsLog/Private/AnalyticsProviderLog.h
--- a/Engine/Plugins/Experimental/StudioTelemetry/Source/AnalyticsLog/Private/AnalyticsProviderLog.h
+++ b/Engine/Plugins/Experimental/StudioTelemetry/Source/AnalyticsLog/Private/AnalyticsProviderLog.h
@@ -5,6 +5,8 @@
 #include "AnalyticsProviderConfigurationDelegate.h"
 #include "Interfaces/IAnalyticsProvider.h"
 
+namespace UE::Tasks { class FPipe; }
+
 /**
  * Implementation of the IAnalyticsProviderET interface that exports telemetry events to a file as Newline - delimited JSON
  * By default,the log file is written to Saved/Telemetry folder of the application.
@@ -41,4 +43,5 @@ private:
 	FString									SessionID;
 	TArray<FAnalyticsEventAttribute>		DefaultEventAttributes;	
 	TUniquePtr<FArchive>					FileWriter;
+	TUniquePtr<UE::Tasks::FPipe>			WriterPipe;
 };
