Added support for StudioTelemetry plugin to execute in non-editor builds.
Added support for LogAnalytics to instantiate in non-editor builds.
Added IsSessionRunning API

diff --git a/Engine/Plugins/Experimental/StudioTelemetry/Source/AnalyticsLog/Public/AnalyticsLog.h b/Engine/Plugins/Experimental/StudioTelemetry/Source/AnalyticsLog/Public/AnalyticsLog.h
--- a/Engine/Plugins/Experimental/StudioTelemetry/Source/AnalyticsLog/Public/AnalyticsLog.h
+++ b/Engine/Plugins/Experimental/StudioTelemetry/Source/AnalyticsLog/Public/AnalyticsLog.h
@@ -2,8 +2,6 @@
 
 #pragma once
 
-#if WITH_EDITOR
-
 #include "AnalyticsProviderConfigurationDelegate.h"
 #include "Containers/Array.h"
 #include "Containers/UnrealString.h"
@@ -46,4 +44,3 @@ private:
 	ANALYTICSLOG_API virtual void ShutdownModule() override;
 };
 
-#endif
\ No newline at end of file
diff --git a/Engine/Plugins/Experimental/StudioTelemetry/Source/StudioTelemetry/Private/AnalyticsProviderMulticast.cpp b/Engine/Plugins/Experimental/StudioTelemetry/Source/StudioTelemetry/Private/AnalyticsProviderMulticast.cpp
--- a/Engine/Plugins/Experimental/StudioTelemetry/Source/StudioTelemetry/Private/AnalyticsProviderMulticast.cpp
+++ b/Engine/Plugins/Experimental/StudioTelemetry/Source/StudioTelemetry/Private/AnalyticsProviderMulticast.cpp
@@ -47,25 +47,43 @@ FAnalyticsProviderMulticast::FAnalyticsProviderMulticast()
 				// Validate the usage type is for this build type
 				if (GConfig->GetString(*ProviderSection, TEXT("UsageType"), UsageType, GEngineIni))
 				{
+					bool IsValidUseCase = false;
+
 #if WITH_EDITOR
 					// Must specify a Editor usage type for this type build
-					if (UsageType.Find(TEXT("Editor")) == INDEX_NONE)
+					if (UsageType.Find(TEXT("Editor")) != INDEX_NONE)
 					{
-						continue;
+						IsValidUseCase |= true;
+					}
+#else 
+					// Must specify a Runtime usage type for all non Editor builds
+					if (UsageType.Find(TEXT("Runtime")) != INDEX_NONE)
+					{
+						IsValidUseCase |= true;
 					}
-#elif WITH_SERVER_CODE
+#endif
+
+#if WITH_SERVER_CODE
 					// Must specify a Server usage type for this type build
-					if (UsageType.Find(TEXT("Server")) == INDEX_NONE)
+					if (UsageType.Find(TEXT("Server")) != INDEX_NONE)
 					{
-						continue;
+						IsValidUseCase |= true;
 					}
-#else
-					// Must specify a Client or Program usage type for this type build
-					if (UsageType.Find(TEXT("Program")) == INDEX_NONE && UsageType.Find(TEXT("Client")) == INDEX_NONE)
+#endif
+
+#if WITH_CLIENT_CODE
+					// Must specify a Client usage type for this type build
+					if (UsageType.Find(TEXT("Client")) != INDEX_NONE)
 					{
-						continue;
+						IsValidUseCase |= true;
 					}
 #endif
+
+					if (IsValidUseCase == false)
+					{
+						// This provider is not valid for this use case
+						continue;
+					}
 				}
 				else
 				{
diff --git a/Engine/Plugins/Experimental/StudioTelemetry/Source/StudioTelemetry/Public/StudioTelemetry.h b/Engine/Plugins/Experimental/StudioTelemetry/Source/StudioTelemetry/Public/StudioTelemetry.h
--- a/Engine/Plugins/Experimental/StudioTelemetry/Source/StudioTelemetry/Public/StudioTelemetry.h
+++ b/Engine/Plugins/Experimental/StudioTelemetry/Source/StudioTelemetry/Public/StudioTelemetry.h
@@ -42,6 +42,9 @@ public:
 
 	/** Access to the tracer for the system*/
 	STUDIOTELEMETRY_API TWeakPtr<IAnalyticsTracer> GetTracer();
+
+	/** Is Session Running */
+	STUDIOTELEMETRY_API bool IsSessionRunning() const;
 	
 	/** Thread safe method to record an event to all registered analytics providers*/
 	STUDIOTELEMETRY_API void RecordEvent(const FString& EventName, const TArray<FAnalyticsEventAttribute>& Attributes = {});
diff --git a/Engine/Plugins/Experimental/StudioTelemetry/StudioTelemetry.uplugin b/Engine/Plugins/Experimental/StudioTelemetry/StudioTelemetry.uplugin
--- a/Engine/Plugins/Experimental/StudioTelemetry/StudioTelemetry.uplugin
+++ b/Engine/Plugins/Experimental/StudioTelemetry/StudioTelemetry.uplugin
@@ -17,7 +17,7 @@
 	"Modules": [
 		{
 			"Name": "AnalyticsLog",
-			"Type": "EditorAndProgram",
+			"Type": "Runtime",
 			"LoadingPhase": "PostConfigInit"
 		},
 		{
@@ -27,7 +27,7 @@
 		},
 		{
 			"Name": "StudioTelemetry",
-			"Type": "EditorAndProgram",
+			"Type": "Runtime",
 			"LoadingPhase": "PostConfigInit"
 		}
 	]
