Added StudioTelemetry::FlushEvents()
Added HttpFlush to AnalyticsProviderMultiCast::FlushEvents() to ensure we flush out any HTTP resuests from any HTTP based Analytics Providers contained within it

diff --git a/Engine/Plugins/Experimental/StudioTelemetry/Source/StudioTelemetry/Private/AnalyticsProviderMulticast.cpp b/Engine/Plugins/Experimental/StudioTelemetry/Source/StudioTelemetry/Private/AnalyticsProviderMulticast.cpp
--- a/Engine/Plugins/Experimental/StudioTelemetry/Source/StudioTelemetry/Private/AnalyticsProviderMulticast.cpp
+++ b/Engine/Plugins/Experimental/StudioTelemetry/Source/StudioTelemetry/Private/AnalyticsProviderMulticast.cpp
@@ -7,6 +7,7 @@
 #include "AnalyticsProviderConfigurationDelegate.h"
 #include "Misc/ConfigCacheIni.h"
 #include "HttpModule.h"
+#include "HttpManager.h"
 
 static FString ProviderSection;
 
@@ -218,10 +219,20 @@ void FAnalyticsProviderMulticast::EndSession()
 
 void FAnalyticsProviderMulticast::FlushEvents()
 {
+	TRACE_CPUPROFILER_EVENT_SCOPE(FAnalyticsProviderMulticast::FlushEvents);
+
 	for (TProviders::TConstIterator it(Providers); it; ++it)
 	{
 		(*it).Value->FlushEvents();
 	}
+
+	// It is quite likely that one of the analytics providers is sending data via the FHttpManager so we
+	// need to flush that as well to make sure that the message is sent. 
+	if (FHttpModule* HttpModule = FModuleManager::GetModulePtr<FHttpModule>("HTTP"))
+	{
+		FHttpManager& HttpManager = HttpModule->GetHttpManager();
+		HttpManager.Flush(EHttpFlushReason::FullFlush);
+	}
 }
 
 void FAnalyticsProviderMulticast::SetRecordEventCallback(OnRecordEvent Callback)
diff --git a/Engine/Plugins/Experimental/StudioTelemetry/Source/StudioTelemetry/Private/StudioTelemetry.cpp b/Engine/Plugins/Experimental/StudioTelemetry/Source/StudioTelemetry/Private/StudioTelemetry.cpp
--- a/Engine/Plugins/Experimental/StudioTelemetry/Source/StudioTelemetry/Private/StudioTelemetry.cpp
+++ b/Engine/Plugins/Experimental/StudioTelemetry/Source/StudioTelemetry/Private/StudioTelemetry.cpp
@@ -70,24 +70,33 @@ void FStudioTelemetry::ShutdownModule()
 
 void FStudioTelemetry::EndSession()
 {
-	OnEndSession.Broadcast();
-
-	// End session for the tracer and the provider
-	if (AnalyticsTracer.IsValid())
+	if (IsSessionRunning())
 	{
-		AnalyticsTracer->EndSession();
-		AnalyticsTracer.Reset();
-	}
+		OnEndSession.Broadcast();
 
-	if (AnalyticsProvider.IsValid())
-	{
-		AnalyticsProvider->EndSession();
-		AnalyticsProvider.Reset();
+		// End session for the tracer and the provider
+		if (AnalyticsTracer.IsValid())
+		{
+			AnalyticsTracer->EndSession();
+			AnalyticsTracer.Reset();
+		}
+
+		if (AnalyticsProvider.IsValid())
+		{
+			AnalyticsProvider->EndSession();
+			AnalyticsProvider->FlushEvents();
+			AnalyticsProvider.Reset();
+		}
 
 		UE_LOG(LogStudioTelemetry, Log, TEXT("Ended StudioTelemetry Session"));
 	}
 }
 
+bool FStudioTelemetry::IsSessionRunning() const
+{
+	return AnalyticsProvider.IsValid();
+}
+
 void FStudioTelemetry::LoadConfiguration()
 {
 	const FString TelemetryConfigurationSection(TEXT("StudioTelemetry.Config"));
@@ -272,6 +281,15 @@ void FStudioTelemetry::RecordEvent(const FName CategoryName, const FString& Even
 	OnRecordEvent.Broadcast(EventName, Attributes);
 }
 
+void FStudioTelemetry::FlushEvents()
+{
+	if (AnalyticsProvider.IsValid())
+	{
+		FScopeLock ScopeLock(&CriticalSection);
+		AnalyticsProvider->FlushEvents();
+	}
+}
+
 void FStudioTelemetry::RecordEventToProvider(const FString& ProviderName, const FString& EventName, const TArray<FAnalyticsEventAttribute>& Attributes)
 {
 	FScopeLock ScopeLock(&CriticalSection);
diff --git a/Engine/Plugins/Experimental/StudioTelemetry/Source/StudioTelemetry/Public/StudioTelemetry.h b/Engine/Plugins/Experimental/StudioTelemetry/Source/StudioTelemetry/Public/StudioTelemetry.h
--- a/Engine/Plugins/Experimental/StudioTelemetry/Source/StudioTelemetry/Public/StudioTelemetry.h
+++ b/Engine/Plugins/Experimental/StudioTelemetry/Source/StudioTelemetry/Public/StudioTelemetry.h
@@ -55,6 +55,9 @@ public:
 	/** Thread safe method to record an event to the specifically named analytics provider */
 	STUDIOTELEMETRY_API void RecordEventToProvider(const FString& ProviderName, const FString& EventName, const TArray<FAnalyticsEventAttribute>& Attributes = {});
 
+	/** Thread safe method to flush all events on all registered analytics providers*/
+	STUDIOTELEMETRY_API void FlushEvents();
+
 	/** Start a new span specifying the parent*/
 	STUDIOTELEMETRY_API TSharedPtr<IAnalyticsSpan> StartSpan(const FName Name, const TArray<FAnalyticsEventAttribute>& AdditionalAttributes = {});
 
